Added the test to validate if a texture is a normal map to the interchange generic pipeline.
Added some function to query the texture compiler to see if a texture is being compiled.

diff --git a/Engine/Plugins/Runtime/Interchange/Source/Pipelines/InterchangePipelines.Build.cs b/Engine/Plugins/Runtime/Interchange/Source/Pipelines/InterchangePipelines.Build.cs
--- a/Engine/Plugins/Runtime/Interchange/Source/Pipelines/InterchangePipelines.Build.cs
+++ b/Engine/Plugins/Runtime/Interchange/Source/Pipelines/InterchangePipelines.Build.cs
@@ -31,6 +31,7 @@ namespace UnrealBuildTool.Rules
 					new string[]
 					{
 						"PhysicsUtilities",
+						"TextureUtilitiesCommon",
 					}
 				);
 			}
diff --git a/Engine/Plugins/Runtime/Interchange/Source/Pipelines/Private/InterchangeGenericAssetsPipeline.cpp b/Engine/Plugins/Runtime/Interchange/Source/Pipelines/Private/InterchangeGenericAssetsPipeline.cpp
--- a/Engine/Plugins/Runtime/Interchange/Source/Pipelines/Private/InterchangeGenericAssetsPipeline.cpp
+++ b/Engine/Plugins/Runtime/Interchange/Source/Pipelines/Private/InterchangeGenericAssetsPipeline.cpp
@@ -1,10 +1,9 @@
 // Copyright Epic Games, Inc. All Rights Reserved. 
+
 #include "InterchangeGenericAssetsPipeline.h"
 
 #include "Animation/Skeleton.h"
 #include "CoreMinimal.h"
-#include "Engine/TextureCube.h"
-#include "Engine/TextureLightProfile.h"
 #include "InterchangeMaterialFactoryNode.h"
 #include "InterchangeMaterialNode.h"
 #include "InterchangeMeshNode.h"
@@ -17,14 +16,7 @@
 #include "InterchangeSourceData.h"
 #include "InterchangeStaticMeshFactoryNode.h"
 #include "InterchangeStaticMeshLodDataNode.h"
-#include "InterchangeTexture2DArrayFactoryNode.h"
-#include "InterchangeTexture2DArrayNode.h"
-#include "InterchangeTexture2DNode.h"
-#include "InterchangeTextureCubeFactoryNode.h"
-#include "InterchangeTextureCubeNode.h"
 #include "InterchangeTextureFactoryNode.h"
-#include "InterchangeTextureLightProfileFactoryNode.h"
-#include "InterchangeTextureLightProfileNode.h"
 #include "InterchangeTextureNode.h"
 #include "Misc/Paths.h"
 #include "Nodes/InterchangeBaseNode.h"
@@ -33,59 +25,6 @@
 #include "UObject/Object.h"
 #include "UObject/ObjectMacros.h"
 
-namespace UE::Interchange::Private
-{
-	UClass* GetDefaultFactoryClassFromTextureNodeClass(UClass* NodeClass)
-	{
-		if (UInterchangeTexture2DNode::StaticClass() == NodeClass)
-		{
-			return UInterchangeTextureFactoryNode::StaticClass();
-		}
-
-		if (UInterchangeTextureCubeNode::StaticClass() == NodeClass)
-		{
-			return UInterchangeTextureCubeFactoryNode::StaticClass();
-		}
-
-		if (UInterchangeTexture2DArrayNode::StaticClass() == NodeClass)
-		{
-			return UInterchangeTexture2DArrayFactoryNode::StaticClass();
-		}
-
-		if (UInterchangeTextureLightProfileNode::StaticClass() == NodeClass)
-		{
-			return UInterchangeTextureLightProfileFactoryNode::StaticClass();
-		}
-
-		return nullptr;
-	}
-
-	UClass* GetDefaultAssetClassFromFactoryClass(UClass* NodeClass)
-	{
-		if (UInterchangeTextureFactoryNode::StaticClass() == NodeClass)
-		{
-			return UTexture2D::StaticClass();
-		}
-
-		if (UInterchangeTextureCubeFactoryNode::StaticClass() == NodeClass)
-		{
-			return UTextureCube::StaticClass();
-		}
-
-		if (UInterchangeTexture2DArrayFactoryNode::StaticClass() == NodeClass)
-		{
-			return UTexture2DArray::StaticClass();
-		}
-
-		if (UInterchangeTextureLightProfileFactoryNode::StaticClass() == NodeClass)
-		{
-			return UTextureLightProfile::StaticClass();
-		}
-
-		return nullptr;
-	}
-}
-
 void UInterchangeGenericAssetsPipeline::PreDialogCleanup(const FName PipelineStackName)
 {
 	//We always clean the pipeline skeleton when showing the dialog
@@ -135,7 +74,7 @@ bool UInterchangeGenericAssetsPipeline::ExecutePreImportPipeline(UInterchangeBas
 		//import textures
 		for (const UInterchangeTextureNode* TextureNode : TextureNodes)
 		{
-			CreateTextureFactoryNode(TextureNode, UE::Interchange::Private::GetDefaultFactoryClassFromTextureNodeClass(TextureNode->GetClass()));
+			HandleCreationOfTextureFactoryNode(TextureNode);
 		}
 	}
 
@@ -282,11 +221,11 @@ bool UInterchangeGenericAssetsPipeline::ExecutePreImportPipeline(UInterchangeBas
 	return true;
 }
 
-bool UInterchangeGenericAssetsPipeline::ExecutePostImportPipeline(const UInterchangeBaseNodeContainer* InBaseNodeContainer, const FString& NodeKey, UObject* CreatedAsset)
+bool UInterchangeGenericAssetsPipeline::ExecutePostImportPipeline(const UInterchangeBaseNodeContainer* InBaseNodeContainer, const FString& NodeKey, UObject* CreatedAsset, bool bIsAReimport)
 {
 	//We do not use the provided base container since ExecutePreImportPipeline cache it
 	//We just make sure the same one is pass in parameter
-	if (!BaseNodeContainer || !ensure(BaseNodeContainer == InBaseNodeContainer) || !CreatedAsset)
+	if (!InBaseNodeContainer || !ensure(InBaseNodeContainer == InBaseNodeContainer) || !CreatedAsset)
 	{
 		return false;
 	}
@@ -299,50 +238,9 @@ bool UInterchangeGenericAssetsPipeline::ExecutePostImportPipeline(const UInterch
 	//Finish the physics asset import, it need the skeletal mesh render data to create the physics collision geometry
 	PostImportPhysicsAssetImport(CreatedAsset, Node);
 
-	return true;
-}
+	PostImportTextureAssetImport(CreatedAsset, bIsAReimport);
 
-UInterchangeTextureFactoryNode* UInterchangeGenericAssetsPipeline::CreateTextureFactoryNode(const UInterchangeTextureNode* TextureNode, const TSubclassOf<UInterchangeTextureFactoryNode>& FactorySubclass)
-{
-	FString DisplayLabel = TextureNode->GetDisplayLabel();
-	FString NodeUid = UInterchangeTextureFactoryNode::GetTextureFactoryNodeUidFromTextureNodeUid(TextureNode->GetUniqueID());
-	UInterchangeTextureFactoryNode* TextureFactoryNode = nullptr;
-	if (BaseNodeContainer->IsNodeUidValid(NodeUid))
-	{
-		TextureFactoryNode = Cast<UInterchangeTextureFactoryNode>(BaseNodeContainer->GetNode(NodeUid));
-		if (!ensure(TextureFactoryNode))
-		{
-			//Log an error
-		}
-	}
-	else
-	{
-		UClass* FactoryClass = FactorySubclass.Get();
-		if (!ensure(FactoryClass))
-		{
-			// Log an error
-			return nullptr;
-		}
-
-		UClass* TextureClass = UE::Interchange::Private::GetDefaultAssetClassFromFactoryClass(FactoryClass);
-		if (!ensure(TextureClass))
-		{
-			// Log an error
-			return nullptr;
-		}
-
-		TextureFactoryNode = NewObject<UInterchangeTextureFactoryNode>(BaseNodeContainer, FactoryClass);
-		if (!ensure(TextureFactoryNode))
-		{
-			return nullptr;
-		}
-		//Creating a UTexture2D
-		TextureFactoryNode->InitializeTextureNode(NodeUid, DisplayLabel, TextureClass->GetName(), TextureNode->GetDisplayLabel());
-		TextureFactoryNode->SetCustomTranslatedTextureNodeUid(TextureNode->GetUniqueID());
-		BaseNodeContainer->AddNode(TextureFactoryNode);
-		TextureFactoryNodes.Add(TextureFactoryNode);
-	}
-	return TextureFactoryNode;
+	return true;
 }
 
 UInterchangeMaterialFactoryNode* UInterchangeGenericAssetsPipeline::CreateMaterialFactoryNode(const UInterchangeMaterialNode* MaterialNode)
diff --git a/Engine/Plugins/Runtime/Interchange/Source/Pipelines/Private/InterchangeGenericTexturePipeline.cpp b/Engine/Plugins/Runtime/Interchange/Source/Pipelines/Private/InterchangeGenericTexturePipeline.cpp
new file mode 100644
--- /dev/null
+++ b/Engine/Plugins/Runtime/Interchange/Source/Pipelines/Private/InterchangeGenericTexturePipeline.cpp
@@ -0,0 +1,187 @@
+// Copyright Epic Games, Inc. All Rights Reserved. 
+
+#include "InterchangeGenericAssetsPipeline.h"
+
+#include "CoreMinimal.h"
+#include "Engine/Texture.h"
+#include "Engine/TextureCube.h"
+#include "Engine/TextureLightProfile.h"
+#include "InterchangeTexture2DArrayFactoryNode.h"
+#include "InterchangeTexture2DArrayFactoryNode.h"
+#include "InterchangeTexture2DArrayNode.h"
+#include "InterchangeTexture2DArrayNode.h"
+#include "InterchangeTexture2DNode.h"
+#include "InterchangeTexture2DNode.h"
+#include "InterchangeTextureCubeFactoryNode.h"
+#include "InterchangeTextureCubeFactoryNode.h"
+#include "InterchangeTextureCubeNode.h"
+#include "InterchangeTextureCubeNode.h"
+#include "InterchangeTextureFactoryNode.h"
+#include "InterchangeTextureFactoryNode.h"
+#include "InterchangeTextureLightProfileFactoryNode.h"
+#include "InterchangeTextureLightProfileFactoryNode.h"
+#include "InterchangeTextureLightProfileNode.h"
+#include "InterchangeTextureLightProfileNode.h"
+#include "InterchangeTextureNode.h"
+
+#if WITH_EDITOR
+#include "NormalMapIdentification.h"
+#include "TextureCompiler.h"
+#endif //WITH_EDITOR
+
+namespace UE::Interchange::Private
+{
+	UClass* GetDefaultFactoryClassFromTextureNodeClass(UClass* NodeClass)
+	{
+		if (UInterchangeTexture2DNode::StaticClass() == NodeClass)
+		{
+			return UInterchangeTextureFactoryNode::StaticClass();
+		}
+
+		if (UInterchangeTextureCubeNode::StaticClass() == NodeClass)
+		{
+			return UInterchangeTextureCubeFactoryNode::StaticClass();
+		}
+
+		if (UInterchangeTexture2DArrayNode::StaticClass() == NodeClass)
+		{
+			return UInterchangeTexture2DArrayFactoryNode::StaticClass();
+		}
+
+		if (UInterchangeTextureLightProfileNode::StaticClass() == NodeClass)
+		{
+			return UInterchangeTextureLightProfileFactoryNode::StaticClass();
+		}
+
+		return nullptr;
+	}
+
+	UClass* GetDefaultAssetClassFromFactoryClass(UClass* NodeClass)
+	{
+		if (UInterchangeTextureFactoryNode::StaticClass() == NodeClass)
+		{
+			return UTexture2D::StaticClass();
+		}
+
+		if (UInterchangeTextureCubeFactoryNode::StaticClass() == NodeClass)
+		{
+			return UTextureCube::StaticClass();
+		}
+
+		if (UInterchangeTexture2DArrayFactoryNode::StaticClass() == NodeClass)
+		{
+			return UTexture2DArray::StaticClass();
+		}
+
+		if (UInterchangeTextureLightProfileFactoryNode::StaticClass() == NodeClass)
+		{
+			return UTextureLightProfile::StaticClass();
+		}
+
+		return nullptr;
+	}
+
+#if WITH_EDITOR
+	void AdjustTextureForNormalMap(UTexture* Texture, bool bFlipNormalMapGreenChannel)
+	{
+		if (Texture)
+		{
+			Texture->PreEditChange(nullptr);
+			if (UE::NormalMapIdentification::HandleAssetPostImport(Texture))
+			{
+				if (bFlipNormalMapGreenChannel)
+				{
+					Texture->bFlipGreenChannel = true;
+				}
+			}
+			Texture->PostEditChange();
+		}
+	}
+#endif
+}
+
+UInterchangeTextureFactoryNode* UInterchangeGenericAssetsPipeline::HandleCreationOfTextureFactoryNode(const UInterchangeTextureNode* TextureNode)
+{
+	return CreateTextureFactoryNode(TextureNode, UE::Interchange::Private::GetDefaultFactoryClassFromTextureNodeClass(TextureNode->GetClass()));
+}
+
+UInterchangeTextureFactoryNode* UInterchangeGenericAssetsPipeline::CreateTextureFactoryNode(const UInterchangeTextureNode* TextureNode, const TSubclassOf<UInterchangeTextureFactoryNode>& FactorySubclass)
+{
+	FString DisplayLabel = TextureNode->GetDisplayLabel();
+	FString NodeUid = UInterchangeTextureFactoryNode::GetTextureFactoryNodeUidFromTextureNodeUid(TextureNode->GetUniqueID());
+	UInterchangeTextureFactoryNode* TextureFactoryNode = nullptr;
+	if (BaseNodeContainer->IsNodeUidValid(NodeUid))
+	{
+		TextureFactoryNode = Cast<UInterchangeTextureFactoryNode>(BaseNodeContainer->GetNode(NodeUid));
+		if (!ensure(TextureFactoryNode))
+		{
+			//Log an error
+		}
+	}
+	else
+	{
+		UClass* FactoryClass = FactorySubclass.Get();
+		if (!ensure(FactoryClass))
+		{
+			// Log an error
+			return nullptr;
+		}
+
+		UClass* TextureClass = UE::Interchange::Private::GetDefaultAssetClassFromFactoryClass(FactoryClass);
+		if (!ensure(TextureClass))
+		{
+			// Log an error
+			return nullptr;
+		}
+
+		TextureFactoryNode = NewObject<UInterchangeTextureFactoryNode>(BaseNodeContainer, FactoryClass);
+		if (!ensure(TextureFactoryNode))
+		{
+			return nullptr;
+		}
+		//Creating a UTexture2D
+		TextureFactoryNode->InitializeTextureNode(NodeUid, DisplayLabel, TextureClass->GetName(), TextureNode->GetDisplayLabel());
+		TextureFactoryNode->SetCustomTranslatedTextureNodeUid(TextureNode->GetUniqueID());
+		BaseNodeContainer->AddNode(TextureFactoryNode);
+		TextureFactoryNodes.Add(TextureFactoryNode);
+	}
+	return TextureFactoryNode;
+}
+
+void UInterchangeGenericAssetsPipeline::PostImportTextureAssetImport(UObject* CreatedAsset, bool bIsAReimport)
+{
+#if WITH_EDITOR
+	if (!bIsAReimport && bDetectNormalMapTexture)
+	{
+		// Verify if the texture is a normal map
+		if (UTexture* Texture = Cast<UTexture>(CreatedAsset))
+		{
+			// This can create 2 build of the texture (we should revisit this at some point)
+			if (FTextureCompilingManager::Get().IsCompilingTexture(Texture))
+			{
+				TWeakObjectPtr<UTexture> WeakTexturePtr = Texture;
+				TSharedRef<FDelegateHandle> HandlePtr = MakeShared<FDelegateHandle>();
+				HandlePtr.Get() = FTextureCompilingManager::Get().OnTexturePostCompileEvent().AddLambda([this, WeakTexturePtr, HandlePtr](const TArrayView<UTexture* const>&)
+					{
+						if (UTexture* TextureToTest = WeakTexturePtr.Get())
+						{
+							if (FTextureCompilingManager::Get().IsCompilingTexture(TextureToTest))
+							{
+								return;
+							}
+
+							UE::Interchange::Private::AdjustTextureForNormalMap(TextureToTest, bFlipNormalMapGreenChannel);
+						}
+
+						FTextureCompilingManager::Get().OnTexturePostCompileEvent().Remove(HandlePtr.Get());
+					});
+			}
+			else
+			{
+				UE::Interchange::Private::AdjustTextureForNormalMap(Texture, bFlipNormalMapGreenChannel);
+			}
+		}
+	}
+#endif //WITH_EDITOR
+}
+
diff --git a/Engine/Plugins/Runtime/Interchange/Source/Pipelines/Public/InterchangeGenericAssetsPipeline.h b/Engine/Plugins/Runtime/Interchange/Source/Pipelines/Public/InterchangeGenericAssetsPipeline.h
--- a/Engine/Plugins/Runtime/Interchange/Source/Pipelines/Public/InterchangeGenericAssetsPipeline.h
+++ b/Engine/Plugins/Runtime/Interchange/Source/Pipelines/Public/InterchangeGenericAssetsPipeline.h
@@ -177,7 +177,19 @@ public:
 	/** If enable, import the material asset find in the sources. */
 	UPROPERTY(EditAnywhere, BlueprintReadWrite, Category = TEXTURES_CATEGORY)
 	bool bImportTextures = true;
-	
+
+#if WITH_EDITORONLY_DATA
+	/** 
+	 * If enable, after a new import a test will be run to see if the texture is a normal map
+	 * If the texture is a normal map the SRG, CompressionSettings and LODGroup settings will be adjusted.
+	 */
+	UPROPERTY(EditAnywhere, BlueprintReadWrite, Category = TEXTURES_CATEGORY)
+	bool bDetectNormalMapTexture = true;
+
+	/** If enabled, the texture's green channel will be inverted for normal maps. */
+	UPROPERTY(EditAnywhere, Category = TEXTURES_CATEGORY)
+	bool bFlipNormalMapGreenChannel = false;
+#endif
 
 	// END Pre import pipeline properties
 	//////////////////////////////////////////////////////////////////////////
@@ -188,7 +200,7 @@ protected:
 
 	virtual bool ExecutePreImportPipeline(UInterchangeBaseNodeContainer* InBaseNodeContainer, const TArray<UInterchangeSourceData*>& InSourceDatas) override;
 
-	virtual bool ExecutePostImportPipeline(const UInterchangeBaseNodeContainer* InBaseNodeContainer, const FString& NodeKey, UObject* CreatedAsset) override;
+	virtual bool ExecutePostImportPipeline(const UInterchangeBaseNodeContainer* BaseNodeContainer, const FString& NodeKey, UObject* CreatedAsset, bool bIsAReimport) override;
 
 	virtual bool CanExecuteOnAnyThread(EInterchangePipelineTask PipelineTask) override
 	{
@@ -208,8 +220,6 @@ private:
 
 	/** Texture factory assets nodes */
 	TArray<UInterchangeTextureFactoryNode*> TextureFactoryNodes;
-
-	UInterchangeTextureFactoryNode* CreateTextureFactoryNode(const UInterchangeTextureNode* TextureNode, const TSubclassOf<UInterchangeTextureFactoryNode>& FactorySubclass);
 	
 	/** Material translated assets nodes */
 	TArray<UInterchangeMaterialNode*> MaterialNodes;
@@ -296,6 +306,18 @@ private:
 
 	/** Specialize for skeletalmesh */
 	void ImplementUseSourceNameForAssetOptionSkeletalMesh(const int32 MeshesAndAnimsImportedNodeCount);
+
+	/************************************************************************/
+	/* Texture API BEGIN                                              */
+
+	UInterchangeTextureFactoryNode* HandleCreationOfTextureFactoryNode(const UInterchangeTextureNode* TextureNode);
+
+	UInterchangeTextureFactoryNode* CreateTextureFactoryNode(const UInterchangeTextureNode* TextureNode, const TSubclassOf<UInterchangeTextureFactoryNode>& FactorySubclass);
+
+	void PostImportTextureAssetImport(UObject* CreatedAsset, bool bIsAReimport);
+
+	/* Texture API END                                                */
+	/************************************************************************/
 };
 
 
diff --git a/Engine/Source/Editor/UnrealEd/Private/Factories/EditorFactories.cpp b/Engine/Source/Editor/UnrealEd/Private/Factories/EditorFactories.cpp
--- a/Engine/Source/Editor/UnrealEd/Private/Factories/EditorFactories.cpp
+++ b/Engine/Source/Editor/UnrealEd/Private/Factories/EditorFactories.cpp
@@ -4362,7 +4362,14 @@ UObject* UTextureFactory::FactoryCreateBinary
 	}
 
 	// Automatically detect if the texture is a normal map and configure its properties accordingly
-	NormalMapIdentification::HandleAssetPostImport(this, Texture);
+
+	if (!bUsingExistingSettings)
+	{
+		if (UE::NormalMapIdentification::HandleAssetPostImport(Texture))
+		{
+			Texture->bFlipGreenChannel = bFlipNormalMapGreenChannel;
+		}
+	}
 
 	if(IsAutomatedImport())
 	{
diff --git a/Engine/Source/Editor/UnrealEd/Private/NormalMapIdentification.cpp b/Engine/Source/Editor/UnrealEd/Private/NormalMapIdentification.cpp
deleted file mode 100644
diff --git a/Engine/Source/Editor/UnrealEd/Private/NormalMapIdentification.h b/Engine/Source/Editor/UnrealEd/Private/NormalMapIdentification.h
deleted file mode 100644
diff --git a/Engine/Source/Runtime/Engine/Private/TextureCompiler.cpp b/Engine/Source/Runtime/Engine/Private/TextureCompiler.cpp
--- a/Engine/Source/Runtime/Engine/Private/TextureCompiler.cpp
+++ b/Engine/Source/Runtime/Engine/Private/TextureCompiler.cpp
@@ -86,6 +86,28 @@ FName FTextureCompilingManager::GetStaticAssetTypeName()
 	return TEXT("UE-Texture");
 }
 
+bool FTextureCompilingManager::IsCompilingTexture(UTexture* InTexture) const 
+{
+	check(IsInGameThread());
+
+	if (!InTexture)
+	{
+		return false;
+	}
+
+	const TWeakObjectPtr<UTexture> InWeakTexturePtr = InTexture;
+	uint32 Hash = GetTypeHash(InWeakTexturePtr);
+	for (const TSet<TWeakObjectPtr<UTexture>>& Bucket : RegisteredTextureBuckets)
+	{
+		if (Bucket.ContainsByHash(Hash, InWeakTexturePtr))
+		{
+			return true;
+		}
+	}
+
+	return false;
+}
+
 FName FTextureCompilingManager::GetAssetTypeName() const
 {
 	return GetStaticAssetTypeName();
@@ -410,6 +432,7 @@ void FTextureCompilingManager::PostCompilation(TArrayView<UTexture* const> InCom
 			}
 
 			FAssetCompilingManager::Get().OnAssetPostCompileEvent().Broadcast(AssetsData);
+			OnTexturePostCompileEvent().Broadcast(InCompiledTextures);
 		}
 	}
 }
diff --git a/Engine/Source/Runtime/Engine/Public/TextureCompiler.h b/Engine/Source/Runtime/Engine/Public/TextureCompiler.h
--- a/Engine/Source/Runtime/Engine/Public/TextureCompiler.h
+++ b/Engine/Source/Runtime/Engine/Public/TextureCompiler.h
@@ -13,6 +13,8 @@ class UTexture;
 class FQueuedThreadPool;
 enum class EQueuedWorkPriority : uint8;
 
+DECLARE_MULTICAST_DELEGATE_OneParam(FTexturePostCompileEvent, const TArrayView<UTexture* const>&);
+
 class FTextureCompilingManager : IAssetCompilingManager
 {
 public:
@@ -71,6 +73,11 @@ public:
 	/** Get the name of the asset type this compiler handles */
 	ENGINE_API static FName GetStaticAssetTypeName();
 
+	/** Return true if the texture is currently compiled */
+	ENGINE_API bool IsCompilingTexture(UTexture* InTexture) const;
+
+	FTexturePostCompileEvent& OnTexturePostCompileEvent() {return TexturePostCompileEvent;}
+
 private:
 	friend class FAssetCompilingManager;
 
@@ -93,6 +100,9 @@ private:
 	bool bHasShutdown = false;
 	TArray<TSet<TWeakObjectPtr<UTexture>>> RegisteredTextureBuckets;
 	FAsyncCompilationNotification Notification;
+
+	/** Event issued at the end of the compile process */
+	FTexturePostCompileEvent TexturePostCompileEvent;
 };
 
 #endif
\ No newline at end of file
diff --git a/Engine/Source/Runtime/Interchange/Core/Public/InterchangePipelineBase.h b/Engine/Source/Runtime/Interchange/Core/Public/InterchangePipelineBase.h
--- a/Engine/Source/Runtime/Interchange/Core/Public/InterchangePipelineBase.h
+++ b/Engine/Source/Runtime/Interchange/Core/Public/InterchangePipelineBase.h
@@ -49,12 +49,12 @@ public:
 	 * @note - the Interchange manager is calling this function not the virtual one that is call by the default implementation.
 	 */
 	UFUNCTION(BlueprintNativeEvent, BlueprintCallable, Category = "Interchange | Pipeline")
-	bool ScriptedExecutePostImportPipeline(const UInterchangeBaseNodeContainer* BaseNodeContainer, const FString& NodeKey, UObject* CreatedAsset);
+	bool ScriptedExecutePostImportPipeline(const UInterchangeBaseNodeContainer* BaseNodeContainer, const FString& NodeKey, UObject* CreatedAsset, bool bIsAReimport);
 	/** The default implementation (call if the blueprint do not have any implementation) will call the virtual ExecuteImportPipeline */
-	bool ScriptedExecutePostImportPipeline_Implementation(const UInterchangeBaseNodeContainer* BaseNodeContainer, const FString& NodeKey, UObject* CreatedAsset)
+	bool ScriptedExecutePostImportPipeline_Implementation(const UInterchangeBaseNodeContainer* BaseNodeContainer, const FString& NodeKey, UObject* CreatedAsset, bool bIsAReimport)
 	{
 		//By default we call the virtual import pipeline execution
-		return ExecutePostImportPipeline(BaseNodeContainer, NodeKey, CreatedAsset);
+		return ExecutePostImportPipeline(BaseNodeContainer, NodeKey, CreatedAsset, bIsAReimport);
 	}
 
 	/**
@@ -141,7 +141,7 @@ protected:
 	 * The interchange manager is not calling this function directly. It is calling the blueprint native event in case this object is a blueprint derive object.
 	 * By default the scripted implementation is calling this virtual pipeline.
 	 */
-	virtual bool ExecutePostImportPipeline(const UInterchangeBaseNodeContainer* BaseNodeContainer, const FString& NodeKey, UObject* CreatedAsset)
+	virtual bool ExecutePostImportPipeline(const UInterchangeBaseNodeContainer* BaseNodeContainer, const FString& NodeKey, UObject* CreatedAsset, bool bIsAReimport)
 	{
 		return false;
 	}
diff --git a/Engine/Source/Runtime/Interchange/Engine/Private/Tasks/InterchangeTaskPipeline.cpp b/Engine/Source/Runtime/Interchange/Engine/Private/Tasks/InterchangeTaskPipeline.cpp
--- a/Engine/Source/Runtime/Interchange/Engine/Private/Tasks/InterchangeTaskPipeline.cpp
+++ b/Engine/Source/Runtime/Interchange/Engine/Private/Tasks/InterchangeTaskPipeline.cpp
@@ -59,9 +59,10 @@ void UE::Interchange::FTaskPipelinePostImport::DoTask(ENamedThreads::Type Curren
 	UInterchangePipelineBase* PipelineBase = AsyncHelper->Pipelines[PipelineIndex];
 	TArray<FString> NodeUniqueIDs;
 	TArray<UObject*> ImportedObjects;
+	TArray<bool> IsAssetsReimported;
 
 	auto FillImportedObjectsFromSource =
-		[&NodeUniqueIDs, &ImportedObjects, this](FCriticalSection& CriticalSection, const TMap<int32, TArray<UE::Interchange::FImportAsyncHelper::FImportedObjectInfo>>& ImportedInfosPerSource)
+		[&NodeUniqueIDs, &ImportedObjects, &IsAssetsReimported, this](FCriticalSection& CriticalSection, const TMap<int32, TArray<UE::Interchange::FImportAsyncHelper::FImportedObjectInfo>>& ImportedInfosPerSource)
 		{
 			FScopeLock Lock(&CriticalSection);
 			if (ImportedInfosPerSource.Contains(SourceIndex))
@@ -69,10 +70,12 @@ void UE::Interchange::FTaskPipelinePostImport::DoTask(ENamedThreads::Type Curren
 				const TArray<UE::Interchange::FImportAsyncHelper::FImportedObjectInfo>& ImportedInfos = ImportedInfosPerSource.FindChecked(SourceIndex);
 				NodeUniqueIDs.Reserve(ImportedInfos.Num());
 				ImportedObjects.Reserve(ImportedInfos.Num());
+				IsAssetsReimported.Reserve(ImportedInfos.Num());
 				for (const UE::Interchange::FImportAsyncHelper::FImportedObjectInfo& ImportedInfo : ImportedInfos)
 				{
 					NodeUniqueIDs.Add(ImportedInfo.FactoryNode->GetUniqueID());
 					ImportedObjects.Add(ImportedInfo.ImportedObject);
+					IsAssetsReimported.Add(ImportedInfo.bIsReimport);
 				}
 			}
 		};
@@ -98,6 +101,6 @@ void UE::Interchange::FTaskPipelinePostImport::DoTask(ENamedThreads::Type Curren
 	//Call the pipeline outside of the lock, we do this in case the pipeline take a long time. We call it for each asset created by this import
 	for (int32 ObjectIndex = 0; ObjectIndex < ImportedObjects.Num(); ++ObjectIndex)
 	{
-		Pipeline->ScriptedExecutePostImportPipeline(NodeContainer, NodeUniqueIDs[ObjectIndex], ImportedObjects[ObjectIndex]);
+		Pipeline->ScriptedExecutePostImportPipeline(NodeContainer, NodeUniqueIDs[ObjectIndex], ImportedObjects[ObjectIndex], IsAssetsReimported[ObjectIndex]);
 	}
 }
diff --git a/Engine/Source/Runtime/TextureUtilitiesCommon/Private/NormalMapIdentification.cpp b/Engine/Source/Runtime/TextureUtilitiesCommon/Private/NormalMapIdentification.cpp
new file mode 100644
--- /dev/null
+++ b/Engine/Source/Runtime/TextureUtilitiesCommon/Private/NormalMapIdentification.cpp
@@ -0,0 +1,519 @@
+// Copyright Epic Games, Inc. All Rights Reserved.
+
+#include "NormalMapIdentification.h"
+
+#if WITH_EDITOR
+
+#include "Engine/Texture2D.h"
+#include "Framework/Notifications/NotificationManager.h"
+#include "Widgets/Notifications/SNotificationList.h"
+
+#define NORMALMAP_IDENTIFICATION_TIMING	(0)
+
+#define LOCTEXT_NAMESPACE "NormalMapIdentification"
+
+////////////////////////////////////////////////////////////////////////////////
+// Constant values
+namespace
+{
+	// These values may need tuning, but results so far have been good
+
+	// These values are the threshold values for the average vector's
+	// length to be considered within limits as a normal map normal
+	const float NormalMapMinLengthConfidenceThreshold = 0.55f;
+	const float NormalMapMaxLengthConfidenceThreshold = 1.1f;
+
+	// This value is the threshold value for the average vector to be considered
+	// to be going in the correct direction.
+	const float NormalMapDeviationThreshold = 0.8f;
+
+	// Samples from the texture will be taken in blocks of this size^2
+	const int32 SampleTileEdgeLength = 4;
+
+	// We sample up to this many tiles in each axis. Sampling more tiles
+	// will likely be more accurate, but will take longer.
+	const int32 MaxTilesPerAxis = 16;
+
+	// This is used in the comparison with "mid-gray"
+	const float ColorComponentNearlyZeroThreshold = (2.0f / 255.0f);
+
+	// This is used when comparing alpha to zero to avoid picking up sprites
+	const float AlphaComponentNearlyZeroThreshold = (1.0f / 255.0f);
+
+	// These values are chosen to make the threshold colors (from uint8 textures)
+	// discard the top most and bottom most two values, i.e. 0, 1, 254 and 255 on
+	// the assumption that these are likely invalid values for a general normal map
+	const float ColorComponentMinVectorThreshold = (2.0f / 255.0f) * 2.0f - 1.0f;
+	const float ColorComponentMaxVectorThreshold = (253.0f/255.0f) * 2.0f - 1.0f;
+
+	// This is the threshold delta length for a vector to be considered as a unit vector
+	const float NormalVectorUnitLengthDeltaThreshold = 0.45f;
+
+	// Rejected to taken sample ratio threshold.
+	const float RejectedToTakenRatioThreshold = 0.33f;
+}
+
+////////////////////////////////////////////////////////////////////////////////
+// Texture sampler classes
+class NormalMapSamplerBase
+{
+public:
+	NormalMapSamplerBase()
+	: SourceTexture(NULL)
+	, TextureSizeX(0)
+	, TextureSizeY(0)
+	, SourceTextureData(NULL)
+	{
+	}
+
+	~NormalMapSamplerBase()
+	{
+		if ( SourceTexture != NULL )
+		{
+			SourceTexture->Source.UnlockMip(0);
+		}
+	}
+
+	void SetSourceTexture( UTexture* Texture )
+	{
+		SourceTexture = Texture;
+		TextureSizeX = Texture->Source.GetSizeX();
+		TextureSizeY = Texture->Source.GetSizeY();
+		SourceTextureData = Texture->Source.LockMipReadOnly(0);
+	}
+
+	UTexture* SourceTexture;
+	int32 TextureSizeX;
+	int32 TextureSizeY;
+	const uint8* SourceTextureData;
+};
+
+template<int RIdx, int GIdx, int BIdx, int AIdx> class SampleNormalMapPixel8 : public NormalMapSamplerBase
+{
+public:
+	SampleNormalMapPixel8() {}
+	~SampleNormalMapPixel8() {}
+
+	FLinearColor DoSampleColor( int32 X, int32 Y )
+	{
+		FLinearColor Result;
+		const uint8* PixelToSample = SourceTextureData + ((Y * TextureSizeX + X) * 4);
+
+		const float OneOver255 = 1.0f / 255.0f;
+
+		Result.B = (float)PixelToSample[BIdx] * OneOver255;
+		Result.G = (float)PixelToSample[GIdx] * OneOver255;
+		Result.R = (float)PixelToSample[RIdx] * OneOver255;
+		Result.A = (float)PixelToSample[AIdx] * OneOver255;
+
+		return Result;
+	}
+
+	float ScaleAndBiasComponent( float Value ) const
+	{
+		return Value * 2.0f - 1.0f;
+	}
+};
+typedef SampleNormalMapPixel8<2, 1, 0, 3> SampleNormalMapPixelBGRA8;
+typedef SampleNormalMapPixel8<0, 1, 2, 3> SampleNormalMapPixelRGBA8;
+
+class SampleNormalMapPixelRGBA16 : public NormalMapSamplerBase
+{
+public:
+	SampleNormalMapPixelRGBA16() {}
+	~SampleNormalMapPixelRGBA16() {}
+
+	FLinearColor DoSampleColor( int32 X, int32 Y )
+	{
+		FLinearColor Result;
+		const uint8* PixelToSample = SourceTextureData + ((Y * TextureSizeX + X) * 8);
+
+		const float OneOver65535 = 1.0f / 65535.0f;
+
+		Result.R = (float)((uint16*)PixelToSample)[0] * OneOver65535;
+		Result.G = (float)((uint16*)PixelToSample)[1] * OneOver65535;
+		Result.B = (float)((uint16*)PixelToSample)[2] * OneOver65535;
+		Result.A = (float)((uint16*)PixelToSample)[3] * OneOver65535;
+
+		return Result;
+	}
+
+	float ScaleAndBiasComponent( float Value ) const
+	{
+		return Value * 2.0f - 1.0f;
+	}
+};
+class SampleNormalMapPixelF16 : public NormalMapSamplerBase
+{
+public:
+	SampleNormalMapPixelF16() {}
+	~SampleNormalMapPixelF16() {}
+
+	FLinearColor DoSampleColor( int32 X, int32 Y )
+	{
+		FLinearColor Result;
+		const uint8* PixelToSample = SourceTextureData + ((Y * TextureSizeX + X) * 8);
+
+		// this assume the normal map to be in linear (not the case if Photoshop converts a 8bit normalmap to float and saves it as 16bit dds)
+		Result.R = (float)((FFloat16*)PixelToSample)[0];
+		Result.G = (float)((FFloat16*)PixelToSample)[1];
+		Result.B = (float)((FFloat16*)PixelToSample)[2];
+		Result.A = (float)((FFloat16*)PixelToSample)[3];
+
+		return Result;
+	}
+
+	float ScaleAndBiasComponent( float Value ) const
+	{
+		// no need to scale and bias floating point components.
+		return Value;
+	}
+};
+
+////////////////////////////////////////////////////////////////////////////////
+////////////////////////////////////////////////////////////////////////////////
+////////////////////////////////////////////////////////////////////////////////
+template<class SamplerClass> class TNormalMapAnalyzer
+{
+public:
+	TNormalMapAnalyzer()
+	: NumSamplesTaken(0)
+	, NumSamplesRejected(0)
+	, NumSamplesThreshold(0)
+	, AverageColor(0.0f,0.0f,0.0f,0.0f)
+	{
+	}
+
+	/**
+	 * EvaluateSubBlock
+	 * Iterates over all pixels in the specified rectangle, if the resulting pixel
+	 * isn't black, mid grey or would result in X or Y being -1 or +1 then it is
+	 * added to the average color and the number of samples count is incremented.
+	 */
+	void EvaluateSubBlock( int32 Left, int32 Top, int32 Width, int32 Height )
+	{
+		for ( int32 Y=Top; Y != (Top+Height); Y++ )
+		{
+			for ( int32 X=Left; X != (Left+Width); X++ )
+			{
+				FLinearColor ColorSample = Sampler.DoSampleColor( X, Y );
+
+				// Nearly black or transparent pixels don't contribute to the calculation
+				if (FMath::IsNearlyZero(ColorSample.A, AlphaComponentNearlyZeroThreshold) ||
+					ColorSample.IsAlmostBlack())
+				{
+					continue;
+				}
+
+				// Scale and bias, if required, to get a signed vector
+				float Vx = Sampler.ScaleAndBiasComponent( ColorSample.R );
+				float Vy = Sampler.ScaleAndBiasComponent( ColorSample.G );
+				float Vz = Sampler.ScaleAndBiasComponent( ColorSample.B );
+
+				const float Length = FMath::Sqrt(Vx * Vx + Vy * Vy + Vz * Vz);
+				if (Length < ColorComponentNearlyZeroThreshold)
+				{
+					// mid-grey pixels representing (0,0,0) are also not considered as they may be used to denote unused areas
+					continue;
+				}
+
+				// If the vector is sufficiently different in length from a unit vector, consider it invalid.
+				if (FMath::Abs(Length - 1.0f) > NormalVectorUnitLengthDeltaThreshold)
+				{
+					NumSamplesRejected++;
+					continue;
+				}
+
+				// If the vector is pointing backwards then it is an invalid sample, so consider it invalid
+				if (Vz < 0.0f)
+				{
+					NumSamplesRejected++;
+					continue;
+				}
+
+				AverageColor += ColorSample;
+				NumSamplesTaken++;
+			}
+		}
+	}
+
+	/**
+	 * DoesTextureLookLikelyToBeANormalMap
+	 *
+	 * Makes a best guess as to whether a texture represents a normal map or not.
+	 * Will not be 100% accurate, but aims to be as good as it can without usage
+	 * information or relying on naming conventions.
+	 *
+	 * The heuristic takes samples in small blocks across the texture (if the texture
+	 * is large enough). The assumption is that if the texture represents a normal map
+	 * then the average direction of the resulting vector should be somewhere near {0,0,1}.
+	 * It samples in a number of blocks spread out to decrease the chance of hitting a
+	 * single unused/blank area of texture, which could happen depending on uv layout.
+	 *
+	 * Any pixels that are black, mid-gray or have a red or green value resulting in X or Y
+	 * being -1 or +1 are ignored on the grounds that they are invalid values. Artists
+	 * sometimes fill the unused areas of normal maps with color being the {0,0,1} vector,
+	 * but that cannot be relied on - those areas are often black or gray instead.
+	 *
+	 * If the heuristic manages to sample enough valid pixels, the threshold being based
+	 * on the total number of samples it will be looking at, then it takes the average
+	 * vector of all the sampled pixels and checks to see if the length and direction are
+	 * within a specific tolerance. See the namespace at the top of the file for tolerance
+	 * value specifications. If the vector satisfies those tolerances then the texture is
+	 * considered to be a normal map.
+	 */
+	bool DoesTextureLookLikelyToBeANormalMap( UTexture* Texture )
+	{
+		int32 TextureSizeX = Texture->Source.GetSizeX();
+		int32 TextureSizeY = Texture->Source.GetSizeY();
+
+		// Calculate the number of tiles in each axis, but limit the number
+		// we interact with to a maximum of 16 tiles (4x4)
+		int32 NumTilesX = FMath::Min( TextureSizeX / SampleTileEdgeLength, MaxTilesPerAxis );
+		int32 NumTilesY = FMath::Min( TextureSizeY / SampleTileEdgeLength, MaxTilesPerAxis );
+
+		Sampler.SetSourceTexture( Texture );
+
+		if (( NumTilesX > 0 ) &&
+			( NumTilesY > 0 ))
+		{
+			// If texture is large enough then take samples spread out across the image
+			NumSamplesThreshold = (NumTilesX * NumTilesY) * 4; // on average 4 samples per tile need to be valid...
+
+			for ( int32 TileY = 0; TileY < NumTilesY; TileY++ )
+			{
+				int Top = (TextureSizeY / NumTilesY) * TileY;
+
+				for ( int32 TileX = 0; TileX < NumTilesX; TileX++ )
+				{
+					int Left = (TextureSizeX / NumTilesX) * TileX;
+
+					EvaluateSubBlock( Left, Top, SampleTileEdgeLength, SampleTileEdgeLength );
+				}
+			}
+		}
+		else
+		{
+			NumSamplesThreshold = (TextureSizeX * TextureSizeY) / 4;
+
+			// Texture is small enough to sample all texels
+			EvaluateSubBlock( 0, 0, TextureSizeX, TextureSizeY );
+		}
+
+		// if we managed to take a reasonable number of samples then we can evaluate the result
+		if ( NumSamplesTaken >= NumSamplesThreshold )
+		{
+			const float RejectedToTakenRatio = static_cast<float>(NumSamplesRejected) / static_cast<float>(NumSamplesTaken);
+			if ( RejectedToTakenRatio >= RejectedToTakenRatioThreshold )
+			{
+				// Too many invalid samples, probably not a normal map
+				return false;
+			}
+
+			AverageColor /= (float)NumSamplesTaken;
+
+			// See if the resulting vector lies anywhere near the {0,0,1} vector
+			float Vx = Sampler.ScaleAndBiasComponent( AverageColor.R );
+			float Vy = Sampler.ScaleAndBiasComponent( AverageColor.G );
+			float Vz = Sampler.ScaleAndBiasComponent( AverageColor.B );
+
+			float Magnitude = FMath::Sqrt( Vx*Vx + Vy*Vy + Vz*Vz );
+
+			// The normalized value of the Z component tells us how close to {0,0,1} the average vector is
+			float NormalizedZ = Vz / Magnitude;
+
+			// if the average vector is longer than or equal to the min length, shorter than the max length
+			// and the normalized Z value means that the vector is close enough to {0,0,1} then we consider
+			// this a normal map
+			return ((Magnitude >= NormalMapMinLengthConfidenceThreshold) &&
+					(Magnitude < NormalMapMaxLengthConfidenceThreshold) &&
+					(NormalizedZ >= NormalMapDeviationThreshold));
+		}
+
+		// Not enough samples, don't trust the result at all
+		return false;
+	}
+
+	int32 NumSamplesTaken;
+	int32 NumSamplesRejected;
+	int32 NumSamplesThreshold;
+	FLinearColor AverageColor;
+
+	SamplerClass Sampler;
+};
+
+/**
+ * Attempts to evaluate the pixels in the texture to see if it is a normal map
+ *
+ * @param Texture The texture to examine
+ *
+ * @return bool true if the texture is likely a normal map (although it's not necessarily guaranteed)
+ */
+static bool IsTextureANormalMap( UTexture* Texture )
+{
+#if NORMALMAP_IDENTIFICATION_TIMING
+	double StartSeconds = FPlatformTime::Seconds();
+#endif
+
+	// Analyze the source texture to try and figure out if it's a normal map.
+	// First check is to make sure it's an appropriate surface format.
+	ETextureSourceFormat SourceFormat = Texture->Source.GetFormat();
+
+	bool bIsNormalMap = false;
+	switch ( SourceFormat )
+	{
+		// The texture could be a normal map if it's one of these formats
+		case TSF_BGRA8:
+			{
+				TNormalMapAnalyzer<SampleNormalMapPixelBGRA8> Analyzer;
+				bIsNormalMap = Analyzer.DoesTextureLookLikelyToBeANormalMap( Texture );
+			}
+			break;
+		case TSF_RGBA16:
+			{
+				TNormalMapAnalyzer<SampleNormalMapPixelRGBA16> Analyzer;
+				bIsNormalMap = Analyzer.DoesTextureLookLikelyToBeANormalMap( Texture );
+			}
+			break;
+		case TSF_RGBA16F:
+			{
+				TNormalMapAnalyzer<SampleNormalMapPixelF16> Analyzer;
+				bIsNormalMap = Analyzer.DoesTextureLookLikelyToBeANormalMap( Texture );
+			}
+			break;
+		case TSF_RGBA8:
+			{
+				TNormalMapAnalyzer<SampleNormalMapPixelRGBA8> Analyzer;
+				bIsNormalMap = Analyzer.DoesTextureLookLikelyToBeANormalMap( Texture );
+			}
+			break;
+
+		default:
+			// assume the texture is not a normal map
+			break;
+	}
+
+#if NORMALMAP_IDENTIFICATION_TIMING
+	double EndSeconds = FPlatformTime::Seconds();
+
+	FString Msg = FString::Printf( TEXT("NormalMapIdentification took %f seconds to analyze %s"), (EndSeconds-StartSeconds), *Texture->GetFullName() ); 
+
+	GLog->Log(Msg);
+#endif
+
+	return bIsNormalMap;
+}
+
+/** Class to handle callbacks from notifications informing the user a texture was imported as a normal map */
+class NormalMapImportNotificationHandler : public TSharedFromThis<NormalMapImportNotificationHandler>
+{
+public:
+	NormalMapImportNotificationHandler() :
+	  Texture(NULL)
+	{
+	}
+	~NormalMapImportNotificationHandler()
+	{
+	}
+
+	/** This method is invoked when the user clicks the "OK" button on the notification */
+	void OKSetting(TSharedPtr<NormalMapImportNotificationHandler>)
+	{
+		if ( Notification.IsValid() )
+		{
+			Notification.Pin()->SetCompletionState(SNotificationItem::ECompletionState::CS_Success);
+			Notification.Pin()->Fadeout();
+		}
+	}
+
+	/* This method is invoked when the user clicked the "Revert" button on the notification */
+	void RevertSetting(TSharedPtr<NormalMapImportNotificationHandler>)
+	{
+		UTexture2D* Texture2D = Texture.IsValid() ? Cast<UTexture2D>(Texture.Get()) : NULL;
+		if ( Texture2D )
+		{
+			if ( Texture2D->CompressionSettings == TC_Normalmap )
+			{
+				// Must wait until the texture is done with previous operations before changing settings and getting it to rebuild.
+				Texture2D->WaitForPendingInitOrStreaming();
+
+				Texture2D->SetFlags(RF_Transactional);
+				Texture2D->Modify();
+				Texture2D->PreEditChange(NULL);
+				{
+					Texture2D->CompressionSettings = TC_Default;
+					Texture2D->SRGB = true;
+					Texture2D->LODGroup = TEXTUREGROUP_World;
+				}
+
+				Texture2D->PostEditChange();
+			}
+		}
+
+		if ( Notification.IsValid() )
+		{
+			Notification.Pin()->SetCompletionState(SNotificationItem::ECompletionState::CS_Success);
+			Notification.Pin()->Fadeout();
+		}
+	}
+
+	TWeakObjectPtr<UTexture> Texture;
+	TWeakPtr<SNotificationItem> Notification;
+};
+
+bool UE::NormalMapIdentification::HandleAssetPostImport( UTexture* Texture )
+{
+	if( Texture != NULL)
+	{
+		// Try to automatically identify a normal map
+		if ( IsTextureANormalMap( Texture ) )
+		{
+			// Set the compression settings and no gamma correction for a normal map
+			{
+				Texture->SetFlags(RF_Transactional);
+				Texture->Modify();
+				Texture->CompressionSettings = TC_Normalmap;
+				Texture->SRGB = false;
+				Texture->LODGroup = TEXTUREGROUP_WorldNormalMap;
+			}
+
+			// Show the user a notification indicating that this texture will be imported as a normal map.
+			// Offer two options to the user, "OK" dismisses the notification early, "Revert" reverts the settings to that of a diffuse map.
+			TSharedPtr<NormalMapImportNotificationHandler> NormalMapNotificationDelegate(new NormalMapImportNotificationHandler);
+			{
+				NormalMapNotificationDelegate->Texture = Texture;
+
+				// this is a cheat to make sure the notification keeps the callback thing alive while it's active...
+				FText OKText = LOCTEXT("ImportTexture_OKNormalMapSettings", "OK");
+				FText OKTooltipText = LOCTEXT("ImportTexture_OKTooltip", "Accept normal map settings");
+				FText RevertText = LOCTEXT("ImportTexture_RevertNormalMapSettings", "Revert");
+				FText RevertTooltipText = LOCTEXT("ImportTexture_RevertTooltip", "Revert to diffuse map settings");
+
+				FFormatNamedArguments Args;
+				Args.Add( TEXT("TextureName"), FText::FromName(Texture->GetFName()) );
+				FNotificationInfo NormalMapNotification( FText::Format(LOCTEXT("ImportTexture_IsNormalMap", "Texture {TextureName} was imported as a normal map"), Args ) );
+				NormalMapNotification.ButtonDetails.Add(FNotificationButtonInfo(OKText, OKTooltipText, FSimpleDelegate::CreateSP(NormalMapNotificationDelegate.Get(), &NormalMapImportNotificationHandler::OKSetting, NormalMapNotificationDelegate)));
+				NormalMapNotification.ButtonDetails.Add(FNotificationButtonInfo(RevertText, RevertTooltipText, FSimpleDelegate::CreateSP(NormalMapNotificationDelegate.Get(), &NormalMapImportNotificationHandler::RevertSetting, NormalMapNotificationDelegate)));
+				NormalMapNotification.bFireAndForget = true;
+				NormalMapNotification.bUseLargeFont = false;
+				NormalMapNotification.bUseSuccessFailIcons = false;
+				NormalMapNotification.bUseThrobber = false;
+				NormalMapNotification.ExpireDuration = 10.0f;
+
+				NormalMapNotificationDelegate->Notification = FSlateNotificationManager::Get().AddNotification(NormalMapNotification);
+				if ( NormalMapNotificationDelegate->Notification.IsValid() )
+				{
+					NormalMapNotificationDelegate->Notification.Pin()->SetCompletionState(SNotificationItem::CS_Pending);
+				}
+			}
+
+			return true;
+		}
+	}
+
+	return false;
+}
+
+#undef LOCTEXT_NAMESPACE
+
+#endif //WITH_EDITOR
diff --git a/Engine/Source/Runtime/TextureUtilitiesCommon/Public/NormalMapIdentification.h b/Engine/Source/Runtime/TextureUtilitiesCommon/Public/NormalMapIdentification.h
new file mode 100644
--- /dev/null
+++ b/Engine/Source/Runtime/TextureUtilitiesCommon/Public/NormalMapIdentification.h
@@ -0,0 +1,24 @@
+// Copyright Epic Games, Inc. All Rights Reserved.
+
+#pragma once
+
+#if WITH_EDITOR
+
+#include "CoreMinimal.h"
+
+class UTexture;
+class UTextureFactory;
+
+namespace UE::NormalMapIdentification
+{
+	/**
+	 * Handle callback when an asset is imported.
+	 * @param	InFactory	The texture factory being used.
+	 * @param	InTexture	The texture that was imported.
+	 * 
+	 * @return true if the asset was identified as normal
+	 */
+	bool TEXTUREUTILITIESCOMMON_API HandleAssetPostImport( UTexture* InTexture );
+}
+
+#endif ///WITH_EDITOR
diff --git a/Engine/Source/Runtime/TextureUtilitiesCommon/TextureUtilitiesCommon.build.cs b/Engine/Source/Runtime/TextureUtilitiesCommon/TextureUtilitiesCommon.build.cs
--- a/Engine/Source/Runtime/TextureUtilitiesCommon/TextureUtilitiesCommon.build.cs
+++ b/Engine/Source/Runtime/TextureUtilitiesCommon/TextureUtilitiesCommon.build.cs
@@ -7,5 +7,13 @@ public class TextureUtilitiesCommon : ModuleRules
 	public TextureUtilitiesCommon(ReadOnlyTargetRules Target) : base(Target)
 	{
 		PublicDependencyModuleNames.Add("Core");
+
+		PrivateDependencyModuleNames.AddRange(
+			new string[] {
+				"CoreUObject",
+				"Engine",
+				"Slate",
+				"SlateCore",
+			});
 	}
 }
\ No newline at end of file
