Better PNGDataFill that uses higher quality nearest neighbor up to a radius of 4 pels
also clarify when it is used and log about it : only on simple transparency, not on alpha channels
also fix PNG loader was not correctly loading transparency mask on non-palettized RGB 3-channel images

diff --git a/Engine/Source/Runtime/ImageWrapper/Private/Formats/PngImageWrapper.cpp b/Engine/Source/Runtime/ImageWrapper/Private/Formats/PngImageWrapper.cpp
--- a/Engine/Source/Runtime/ImageWrapper/Private/Formats/PngImageWrapper.cpp
+++ b/Engine/Source/Runtime/ImageWrapper/Private/Formats/PngImageWrapper.cpp
@@ -378,6 +378,8 @@ void FPngImageWrapper::UncompressPNGData(const ERGBFormat InFormat, const int32
 				png_set_palette_to_rgb(png_ptr);
 			}
 
+			// really we should just call png_expand() here and remove all these conditionals
+
 			if (((ColorType & PNG_COLOR_MASK_COLOR) == 0) && BitDepth < 8)
 			{
 				png_set_expand_gray_1_2_4_to_8(png_ptr);
@@ -387,14 +389,10 @@ void FPngImageWrapper::UncompressPNGData(const ERGBFormat InFormat, const int32
 			if ((ColorType & PNG_COLOR_MASK_ALPHA) == 0 && (InFormat == ERGBFormat::BGRA || InFormat == ERGBFormat::RGBA))
 			{
 				// png images don't set PNG_COLOR_MASK_ALPHA if they have alpha from a tRNS chunk, but png_set_add_alpha seems to be safe regardless
-				if ((ColorType & PNG_COLOR_MASK_COLOR) == 0)
-				{
-					png_set_tRNS_to_alpha(png_ptr);
-				}
-				else if (ColorType == PNG_COLOR_TYPE_PALETTE)
-				{
-					png_set_tRNS_to_alpha(png_ptr);
-				}
+				png_set_tRNS_to_alpha(png_ptr);
+
+				// note: png_set_tRNS_to_alpha is just an alias for png_expand
+
 				if (InBitDepth == 8)
 				{
 					png_set_add_alpha(png_ptr, 0xff , PNG_FILLER_AFTER);
diff --git a/Engine/Source/Runtime/TextureUtilitiesCommon/Private/TextureImportUtils.cpp b/Engine/Source/Runtime/TextureUtilitiesCommon/Private/TextureImportUtils.cpp
--- a/Engine/Source/Runtime/TextureUtilitiesCommon/Private/TextureImportUtils.cpp
+++ b/Engine/Source/Runtime/TextureUtilitiesCommon/Private/TextureImportUtils.cpp
@@ -3,6 +3,7 @@
 #include "TextureImportUtils.h"
 
 #include "ImageCore.h"
+#include "Async/ParallelFor.h"
 
 namespace UE
 {
@@ -42,8 +43,14 @@ namespace UE
 		}
 
 		/**
-		 * This fills any pixels of a texture with have an alpha value of zero,
+		 * This fills any pixels of a texture with have an alpha value of zero and RGB=white,
 		 * with an RGB from the nearest neighboring pixel which has non-zero alpha.
+
+		   PNG images with "simple transparency" (eg. indexed color transparency) don't store RGB color in the transparent area
+		   libpng decodes those pels are {RGB=white, A=0}
+		   we replace them by filling in the RGB from neighbors
+
+		   note that this does NOT fill in the RGB of PNGs with a full alpha channel.
 		 */
 		template<typename PixelDataType, typename ColorDataType, int32 RIdx, int32 GIdx, int32 BIdx, int32 AIdx>
 		class TPNGDataFill
@@ -55,10 +62,84 @@ namespace UE
 				, TextureWidth(SizeX)
 				, TextureHeight(SizeY)
 			{
+				// libpng decodes simple transparent (binary A or indexed color) as { RGB=white, A = 0}
+
+				if ( sizeof(ColorDataType) == 4 )
+				{
+					WhiteWithZeroAlpha = FColor(255, 255, 255, 0).DWColor();
+				}
+				else
+				{
+					uint16 RGBA[4] = { 0xFFFF,0xFFFF,0xFFFF, 0 };
+					checkSlow( sizeof(ColorDataType) == 8 );
+					checkSlow( sizeof(RGBA) == 8 );
+					memcpy(&WhiteWithZeroAlpha,RGBA, sizeof(ColorDataType));
+				}
+				
+				// falloff weights for near neighbor extrapolation
+				// fast falloff -> just extend neighbor pels out
+				// slow falloff -> blur neighbor pels together
+				for(int r1=0;r1<=NearNeighborRadius;r1++)
+				{
+					for(int r2=0;r2<=NearNeighborRadius;r2++)
+					{
+						int rsqr = r1*r1 + r2*r2;
+						if( rsqr == 0 )
+						{
+							NearNeighborWeights[0][0] = 0.f;
+							continue;
+						}
+						float W = expf( - 1.1f * sqrtf((float)rsqr) );
+						NearNeighborWeights[r1][r2] = W;
+					}
+				}
 			}
 
 			void ProcessData()
 			{
+				// first identify alpha type :
+				bool HasBinaryTransparency=false;
+				for (int64 Y = 0; Y < TextureHeight; ++Y)
+				{
+					const ColorDataType* RowData = (const ColorDataType *)SourceData + Y * TextureWidth;
+
+					for(int64 X = 0; X < TextureWidth; ++X)
+					{
+						if ( IsOpaque(RowData[X]) )
+						{
+						}
+						else if ( RowData[X] == WhiteWithZeroAlpha )
+						{
+							HasBinaryTransparency = true;
+						}
+						else
+						{
+							//HasComplexAlphaChannel = true;
+							UE_LOG(LogCore, Log,  TEXT("PNG has complex alpha channel, will not fill RGB in transparent background"));
+							// do not modify png's with full alpha channels :
+							return;
+						}
+					}
+				}
+
+				if ( ! HasBinaryTransparency )
+				{
+					// all opaque
+					return;
+				}
+
+				UE_LOG(LogCore, Log,  TEXT("PNG has binary transparency, doing fill of RGB in transparent background"));
+
+				// first do good fill with limited distance :
+				// this ensures near pels within NearNeighborRadius get a good neighbor fill for interpolation
+
+				FillFromNearNeighbors();
+
+				// then do simple fill, rows one by one :
+				// this can be a very poor fill, but it's fast for filling large empty areas
+
+				// @todo oodle : fill from nearest row (up or down) rather than always filling downward
+
 				int64 NumZeroedTopRowsToProcess = 0;
 				int64 FillColorRow = -1;
 				for (int64 Y = 0; Y < TextureHeight; ++Y)
@@ -71,7 +152,7 @@ namespace UE
 						}
 						else
 						{
-							NumZeroedTopRowsToProcess = Y;
+							NumZeroedTopRowsToProcess = Y+1;
 						}
 					}
 					else
@@ -81,92 +162,199 @@ namespace UE
 				}
 
 				// Can only fill upwards if image not fully zeroed
-				if (NumZeroedTopRowsToProcess > 0 && NumZeroedTopRowsToProcess + 1 < TextureHeight)
+				if (NumZeroedTopRowsToProcess > 0 && NumZeroedTopRowsToProcess < TextureHeight)
+				{
+					for (int64 Y = 0; Y < NumZeroedTopRowsToProcess; ++Y)
+					{
+						// fill row at Y from row at NumZeroedTopRowsToProcess
+						FillRowColorPixels(NumZeroedTopRowsToProcess, Y);
+					}
+				}
+
+				// change all pixels we made in FillFromNearNeighbors and change to A=0 :
+				
+				for (int64 Y = 0; Y < TextureHeight; ++Y)
 				{
-					for (int64 Y = 0; Y <= NumZeroedTopRowsToProcess; ++Y)
+					ColorDataType * ImageRow = (ColorDataType *)SourceData + Y * TextureWidth;
+					for (int64 X = 0; X < TextureWidth; ++X)
 					{
-						FillRowColorPixels(NumZeroedTopRowsToProcess + 1, Y);
+						const ColorDataType Color = ImageRow[X];
+						if ( Color == WhiteWithZeroAlpha || IsOpaque(Color) )
+						{
+							continue;
+						}
+
+						ImageRow[X] = MakeColorWithZeroAlpha(Color);
 					}
 				}
 			}
 
-			/* returns False if requires further processing because entire row is filled with zeroed alpha values */
-			bool ProcessHorizontalRow(int64 Y)
+			static bool IsOpaque(const ColorDataType InColor)
 			{
-				// only wipe out colors that are affected by png turning valid colors white if alpha = 0
-				ColorDataType WhiteWithZeroAlpha;
 				if ( sizeof(ColorDataType) == 4 )
 				{
-					WhiteWithZeroAlpha = FColor(255, 255, 255, 0).DWColor();
+					return InColor >= 0xFF000000U;
+				}
+				else if ( sizeof(ColorDataType) == 8 )
+				{
+					return InColor >= 0xFFFF000000000000ULL;
 				}
 				else
 				{
-					check( sizeof(ColorDataType) == 8 );
-					uint16 RGBA[4] = { 0xFFFF,0xFFFF,0xFFFF, 0 };
-					check( sizeof(RGBA) == 8 );
-					memcpy(&WhiteWithZeroAlpha,RGBA, sizeof(ColorDataType));
+					check(false);
 				}
+			}
 
-				// Left -> Right
-				int64 NumLeftmostZerosToProcess = 0;
-				const PixelDataType* FillColor = nullptr;
-				for (int64 X = 0; X < TextureWidth; ++X)
+			static ColorDataType MakeColorWithZeroAlpha(const ColorDataType InColor)
+			{
+				// take the RGB from InColor
+				// set A to zero
+				// return that
+
+				if ( sizeof(ColorDataType) == 4 )
 				{
-					PixelDataType* PixelData = SourceData + (Y * TextureWidth + X) * 4;
-					ColorDataType* ColorData = reinterpret_cast<ColorDataType*>(PixelData);
+					return InColor & 0xFFFFFFU;
+				}
+				else if ( sizeof(ColorDataType) == 8 )
+				{
+					return InColor & 0xFFFFFFFFFFFFULL;
+				}
+				else
+				{
+					check(false);
+				}
+			}
+			
+			static ColorDataType MakeColorOpaque(const ColorDataType InColor)
+			{
+				// take the RGB from InColor
+				// set A to opaque
+				// return that
+
+				if ( sizeof(ColorDataType) == 4 )
+				{
+					return InColor | 0xFF000000U;
+				}
+				else if ( sizeof(ColorDataType) == 8 )
+				{
+					return InColor | 0xFFFF000000000000ULL;
+				}
+				else
+				{
+					check(false);
+				}
+			}
+
+
+			/* returns False if requires further processing because entire row is filled with zeroed alpha values */
+			bool ProcessHorizontalRow(int64 Y)
+			{		
+				ColorDataType* RowData = (ColorDataType *)SourceData + Y * TextureWidth;
+				int64 X = 0;
+
+				if ( RowData[0] == WhiteWithZeroAlpha )
+				{
+					// transparent run at start of row
+					// find X which is the first opaque pel
 
-					if (*ColorData == WhiteWithZeroAlpha)
+					for(;;)
 					{
-						if (FillColor)
+						if ( RowData[X] != WhiteWithZeroAlpha )
 						{
-							PixelData[RIdx] = FillColor[RIdx];
-							PixelData[GIdx] = FillColor[GIdx];
-							PixelData[BIdx] = FillColor[BIdx];
+							break;
 						}
-						else
-						{
-							// Mark pixel as needing fill
-							*ColorData = 0;
 
-							// Keep track of how many pixels to fill starting at beginning of row
-							NumLeftmostZerosToProcess = X;
+						X++;
+
+						if ( X == TextureWidth )
+						{
+							// whole row was transparent
+							return false;
 						}
 					}
-					else
+
+					check( X < TextureWidth );
+					check( RowData[X] != WhiteWithZeroAlpha );
+
+					// RowData[X] is opaque
+					// fill initial run from it
+					ColorDataType FillColor = MakeColorWithZeroAlpha(RowData[X]);
+					for(int64 FillX=0;FillX<X;FillX++)
 					{
-						FillColor = PixelData;
+						RowData[FillX] = FillColor;
 					}
 				}
 
-				if (NumLeftmostZerosToProcess == 0)
+				for(;;)
 				{
-					// No pixels left that are zero
-					return true;
-				}
+					// we're in an opaque run, step to end of opaque run :
+					check( RowData[X] != WhiteWithZeroAlpha );
 
-				if (NumLeftmostZerosToProcess + 1 >= TextureWidth)
-				{
-					// All pixels in this row are zero and must be filled using rows above or below
-					return false;
-				}
+					while( RowData[X] != WhiteWithZeroAlpha )
+					{
+						X++;
 
-				// Fill using non zero pixel immediately to the right of the beginning series of zeros
-				FillColor = SourceData + (Y * TextureWidth + NumLeftmostZerosToProcess + 1) * 4;
+						if ( X == TextureWidth )
+						{
+							//reached end in opaque run, done.
+							return true;
+						}
+					}
 
-				// Fill zero pixels found at beginning of row that could not be filled during the Left to Right pass
-				for (int64 X = 0; X <= NumLeftmostZerosToProcess; ++X)
-				{
-					PixelDataType* PixelData = SourceData + (Y * TextureWidth + X) * 4;
-					PixelData[RIdx] = FillColor[RIdx];
-					PixelData[GIdx] = FillColor[GIdx];
-					PixelData[BIdx] = FillColor[BIdx];
+					// X is transparent
+					check( X > 0 );
+					check( RowData[X] == WhiteWithZeroAlpha );
+
+					int64 FirstTransparent = X;
+					while( RowData[X] == WhiteWithZeroAlpha )
+					{
+						X++;
+						
+						if ( X == TextureWidth )
+						{
+							//reached end in transparent run
+							// fill right-only transparent run from left :
+							
+							ColorDataType FillColor = MakeColorWithZeroAlpha(RowData[FirstTransparent-1]);
+							for(int64 FillX=FirstTransparent;FillX<TextureWidth;FillX++)
+							{
+								RowData[FillX] = FillColor;
+							}
+
+							return true;
+						}
+					}
+					
+					// RowData[X] is opaque
+					check( X < TextureWidth );
+					check( RowData[X] != WhiteWithZeroAlpha );
+					// transparent run [FirstTransparent,X) with opaque on each side
+
+					// fill left half from Left, right half from Right
+					int64 MidX = (FirstTransparent + X)/2;
+					ColorDataType FillLeft  = MakeColorWithZeroAlpha(RowData[FirstTransparent-1]);
+					ColorDataType FillRight = MakeColorWithZeroAlpha(RowData[X]);
+					
+					for(int64 FillX=FirstTransparent;FillX<MidX;FillX++)
+					{
+						RowData[FillX] = FillLeft;
+					}
+					for(int64 FillX=MidX;FillX<X;FillX++)
+					{
+						RowData[FillX] = FillRight;
+					}
+
+					// X is opaque, repeat
 				}
 
+				// can't get here
+
 				return true;
 			}
 
 			void FillRowColorPixels(int64 FillColorRow, int64 Y)
 			{
+				// fill row Y from row FillColorRow , copying RGB only, leave A alone (A = 0 in dest)
 				for (int64 X = 0; X < TextureWidth; ++X)
 				{
 					const PixelDataType* FillColor = SourceData + (FillColorRow * TextureWidth + X) * 4;
@@ -176,10 +364,159 @@ namespace UE
 					PixelData[BIdx] = FillColor[BIdx];
 				}
 			}
+			
+			static FLinearColor MakeLinearColor(const ColorDataType InColor)
+			{
+				if ( sizeof(ColorDataType) == 4 )
+				{
+					// does SRGB gamma conversion :
+					return FLinearColor( FColor(InColor) );
+				}
+				else
+				{
+					uint16 RGBA[4];
+					checkSlow( sizeof(ColorDataType) == 8 );
+					checkSlow( sizeof(RGBA) == 8 );
+					memcpy(&RGBA,&InColor,sizeof(ColorDataType));
+
+					return FLinearColor( 
+						RGBA[0] * (1.f/0xFFFF),
+						RGBA[1] * (1.f/0xFFFF),
+						RGBA[2] * (1.f/0xFFFF),
+						RGBA[3] * (1.f/0xFFFF));
+				}
+			}
+
+			static ColorDataType MakeColorFromLinear(const FLinearColor InColor)
+			{
+				if ( sizeof(ColorDataType) == 4 )
+				{
+					return InColor.ToFColorSRGB().DWColor();
+				}
+				else
+				{
+					uint16 RGBA[4];
+					checkSlow( sizeof(ColorDataType) == 8 );
+					checkSlow( sizeof(RGBA) == 8 );
+
+					RGBA[0] = FColor::QuantizeUNormFloatTo16( InColor.R );
+					RGBA[1] = FColor::QuantizeUNormFloatTo16( InColor.G );
+					RGBA[2] = FColor::QuantizeUNormFloatTo16( InColor.B );
+					RGBA[3] = FColor::QuantizeUNormFloatTo16( InColor.A );
+
+					ColorDataType Ret;
+					memcpy(&Ret,RGBA,sizeof(ColorDataType));
+					return Ret;
+				}
+			}
+
+			ColorDataType GetFilledFromNearNeighbors(int64 CenterX,int64 CenterY) const
+			{
+				// look in the local neighborhood around CenterX,CenterY
+				// find any opaque pixels that were in the source which we can expand from
+				// (do not use opaque pixels made from previous NearNeighbor fills)
+				// weight contribution by distance, falling off fast
+				// (mainly we want the case of diagonal to average from +X and +Y)
+
+				int64 LowX = FMath::Max(0,CenterX - NearNeighborRadius);
+				int64 LowY = FMath::Max(0,CenterY - NearNeighborRadius);
+				int64 HighX = FMath::Min(TextureWidth -1,CenterX + NearNeighborRadius);
+				int64 HighY = FMath::Min(TextureHeight-1,CenterY + NearNeighborRadius);
+				// search in [LowX,HighX] inclusive
+				
+				const ColorDataType* ImageColors = (const ColorDataType *)SourceData;
+				float TotalWeight = 0.f;
+				FLinearColor TotalColor(0,0,0,0);
+
+				check( ImageColors[ CenterX + CenterY * TextureWidth ] == WhiteWithZeroAlpha );
+
+				for(int64 Y=LowY;Y<=HighY;Y++)
+				{
+					const float * RowNearNeighborWeights = NearNeighborWeights[ FMath::Abs(Y - CenterY) ];
+
+					for(int64 X=LowX;X<=HighX;X++)
+					{
+						const ColorDataType Color = ImageColors[ X + Y * TextureWidth ];
+
+						// check IsOpaque rather than == WhiteWithZeroAlpha so that we can work in-place
+						if ( IsOpaque(Color) )
+						{
+							float W = RowNearNeighborWeights[ FMath::Abs(X-CenterX) ];
+							// accumulate weighted color
+							TotalWeight += W;
+							TotalColor += W * MakeLinearColor(Color);
+						}
+					}
+				}
+
+				if ( TotalWeight == 0.f )
+				{
+					// found nothing nearby
+					return WhiteWithZeroAlpha;
+				}
+
+				TotalColor *= 1.f / TotalWeight;
+				TotalColor.A = 0.5f; // anything other than opaque or zero
+
+				// we return an A that's not opaque so we won't be read from in this pass
+				// later we will change this A to zero
+
+				return MakeColorFromLinear(TotalColor);
+			}
+
+			void FillFromNearNeighbors()
+			{
+				// higher quality, slower neighbor fill with a limited distance of NearNeighborRadius
+
+				ColorDataType* ImageColors = (ColorDataType *)SourceData;
+								
+				// this is trivially parallelizable because each pixel write op is independent
+				//	and we don't read results that we write in this pass (because the A value excludes new writes)
+
+				int32 NumRowsEachJob;
+				int32 NumJobs = ImageParallelForComputeNumJobsForRows(NumRowsEachJob, TextureWidth, TextureHeight);
+
+				ParallelFor(TEXT("TextureImportUtils.FillFromNearNeighbors.PF"), NumJobs, 1, [&](int32 Index)
+				{
+					int64 StartIndex = Index * NumRowsEachJob;
+					int64 EndIndex = FMath::Min(StartIndex + NumRowsEachJob, TextureHeight);
+
+					// using a scratch output row is not strictly required, but reduces cache thrashing between threads :
+					//  (if single-threaded, no reason to use it, just write ImageRow[x] immediately)
+					TArray<ColorDataType> ScratchRowArray;
+					ScratchRowArray.SetNum(TextureWidth);
+					ColorDataType * ScratchRow = &ScratchRowArray[0];
+
+					for (int64 Y = StartIndex; Y < EndIndex; ++Y)
+					{
+						ColorDataType * ImageRow = ImageColors + Y * TextureWidth;
+						for (int64 X = 0; X < TextureWidth; ++X)
+						{
+							//@todo Oodle : we could more quickly detect large areas where no fill within NearNeighborRadius is possible
+							
+							if ( ImageRow[X] == WhiteWithZeroAlpha )
+							{
+								// could write to ImageRow[X] immediately, but using ScratchRow reduces cache sharing across cores
+								ScratchRow[X] = GetFilledFromNearNeighbors(X,Y);
+							}
+							else
+							{
+								ScratchRow[X] = ImageRow[X];
+							}
+						}
+						memcpy(ImageRow,ScratchRow,TextureWidth*sizeof(ColorDataType));
+					}
+				});
+			}
 
 			PixelDataType* SourceData;
 			int64 TextureWidth;
 			int64 TextureHeight;
+			
+			ColorDataType WhiteWithZeroAlpha;
+
+			enum { NearNeighborRadius = 4 };
+			float NearNeighborWeights[NearNeighborRadius+1][NearNeighborRadius+1];
 		};
 
 		void FillZeroAlphaPNGData(int32 SizeX, int32 SizeY, ETextureSourceFormat SourceFormat, uint8* SourceData)
