TextureImportSettings : allow choice of how floating point textures are compressed by default
legacy default behavior is to always use HDR F16
two new good options are to use BC6H (for games) or to use F32 formats when the source is 32-bit (for VFX)

diff --git a/Engine/Source/Runtime/TextureUtilitiesCommon/Private/TextureImportSettings.cpp b/Engine/Source/Runtime/TextureUtilitiesCommon/Private/TextureImportSettings.cpp
--- a/Engine/Source/Runtime/TextureUtilitiesCommon/Private/TextureImportSettings.cpp
+++ b/Engine/Source/Runtime/TextureUtilitiesCommon/Private/TextureImportSettings.cpp
@@ -83,6 +83,50 @@ namespace UE::TextureUtilitiesCommon
 		{
 			Texture->bUseNewMipFilter = true;
 		}
+
+		// the pipeline before here will have set floating point textures to TC_HDR
+		//	could alternatively check Texture->HasHDRSource
+		if ( Texture->CompressionSettings == TC_HDR )
+		{
+			if ( Settings->CompressedFormatForFloatTextures == ETextureImportFloatingPointFormat::HDRCompressed_BC6 )
+			{
+				// use BC6H
+				Texture->CompressionSettings = TC_HDR_Compressed;
+			}
+			else if ( Settings->CompressedFormatForFloatTextures == ETextureImportFloatingPointFormat::HDR_F32_or_F16 )
+			{
+				// set output format to match source format
+				ETextureSourceFormat TSF = Texture->Source.GetFormat();
+				if ( TSF == TSF_RGBA32F )
+				{
+					Texture->CompressionSettings = 	TC_HDR_F32;	
+				}
+				else if ( TSF == TSF_R32F )
+				{
+					Texture->CompressionSettings = 	TC_SingleFloat;
+				}
+				else if ( TSF == TSF_R16F )
+				{
+					Texture->CompressionSettings = 	TC_HalfFloat;
+				}
+				else
+				{
+					// else leave TC_HDR
+					check( Texture->CompressionSettings == TC_HDR );
+				}
+			}
+			else if ( Settings->CompressedFormatForFloatTextures == ETextureImportFloatingPointFormat::HDR_F16 )
+			{
+				// always use F16 HDR (legacy behavior)
+				// leave TC_HDR
+				check( Texture->CompressionSettings == TC_HDR );
+			}
+			else
+			{
+				// all cases should have been handled
+				checkNoEntry();
+			}
+		}
 	}
 
 }
diff --git a/Engine/Source/Runtime/TextureUtilitiesCommon/Public/TextureImportSettings.h b/Engine/Source/Runtime/TextureUtilitiesCommon/Public/TextureImportSettings.h
--- a/Engine/Source/Runtime/TextureUtilitiesCommon/Public/TextureImportSettings.h
+++ b/Engine/Source/Runtime/TextureUtilitiesCommon/Public/TextureImportSettings.h
@@ -11,6 +11,21 @@
 
 struct FPropertyChangedEvent;
 
+/* What CompressionSettings runtime format should imported floating point textures use
+ */
+UENUM()
+enum class ETextureImportFloatingPointFormat : uint8
+{
+	/* Use "HDR" RGBA16F */
+	HDR_F16 = 0,
+	/* Use "HDRCompressed" , BC6H */
+	HDRCompressed_BC6,
+	/* Use 32-bit float formats if the source is 32-bit, otherwise use 16-bit HDR */
+	HDR_F32_or_F16,
+
+	PreviousDefault = HDR_F16  UMETA(Hidden) // legacy behavior
+};
+
 UCLASS(config=Editor, defaultconfig, meta=(DisplayName="Texture Import"))
 class TEXTUREUTILITIESCOMMON_API UTextureImportSettings : public UDeveloperSettings
 {
@@ -32,6 +47,11 @@ public:
 		DisplayName = "Turn on fast mip generation filter",
 		ToolTip = "Use the fast mip filter on new textures; it is recommended, but can be turned off to maintain legacy behavior. This setting is applied to newly imported textures, it does not affect existing textures in the project."))
 	bool bEnableFastMipFilter = true;
+	
+	UPROPERTY(config, EditAnywhere, Category=ImportSettings, meta = (
+		DisplayName = "CompressionFormat to use for new float textures",
+		ToolTip = "Optionally use HDRCompressed (BC6H), or 32-bit adaptively, instead of HDR (RGBA16F) for floating point textures.  This setting is applied to newly imported textures, it does not affect existing textures in the project."))
+	ETextureImportFloatingPointFormat CompressedFormatForFloatTextures = ETextureImportFloatingPointFormat::PreviousDefault;
 
 	//~ Begin UObject Interface
 
