Adding an utility to resize the source data of some textures (Those that don't use  layers or slices (except for the cube map)). Only tested and supported with the Texture2D and TextureCube at the moment.

Added a function to the textures to let something external use their texture mip generation to downsize a 2DImage closer to a desired target resolution.

Added an utility that try to reduce the resolution of the source data of texture closer what is needed for it's rendering/in-game resolution.

#rnx

diff --git a/Engine/Source/Runtime/Engine/Classes/Engine/Texture.h b/Engine/Source/Runtime/Engine/Classes/Engine/Texture.h
--- a/Engine/Source/Runtime/Engine/Classes/Engine/Texture.h
+++ b/Engine/Source/Runtime/Engine/Classes/Engine/Texture.h
@@ -1506,7 +1506,7 @@ public:
 	 * Get the dimensions of the largest mip of the texture when built for the target platform
 	 *   accounting for LODBias and other constraints
 	 */
-	void GetBuiltTextureSize( const ITargetPlatform* TargetPlatform , int32 & OutSizeX, int32 & OutSizeY ) const;
+	ENGINE_API void GetBuiltTextureSize( const ITargetPlatform* TargetPlatform , int32 & OutSizeX, int32 & OutSizeY ) const;
 
 	/**
 	 * Serializes cooked platform data.
@@ -1848,6 +1848,12 @@ public:
 	* this is not for the current RHI (which may have a lower limit), this is for a Texture in general
 	*/
 	ENGINE_API static int32 GetMaximumDimensionOfNonVT();
+
+	/*
+	 * Downsize the 2D Image with the build setting of the texture.
+	 * Try to get as close as it can to the target resolution but it will stay above it if it can reach it
+	 */
+	ENGINE_API bool DownsizeImageUsingTextureSettings(const ITargetPlatform* TargetPlatform, FImage& InOutImage, int32 TargetSize, int32 LayerIndex);
 #endif
 
 protected:
diff --git a/Engine/Source/Runtime/Engine/Private/Texture.cpp b/Engine/Source/Runtime/Engine/Private/Texture.cpp
--- a/Engine/Source/Runtime/Engine/Private/Texture.cpp
+++ b/Engine/Source/Runtime/Engine/Private/Texture.cpp
@@ -3230,8 +3230,9 @@ void UTexture::GetBuiltTextureSize( const ITargetPlatform* TargetPlatform , int3
 	int32 SizeX,SizeY;
 
 #if WITH_EDITORONLY_DATA
-	SizeX = Source.GetSizeX();
-	SizeY = Source.GetSizeY();
+	FIntPoint SourceSize = Source.GetLogicalSize();
+	SizeX = SourceSize.X;
+	SizeY = SourceSize.Y;
 
 	if (PowerOfTwoMode == ETexturePowerOfTwoSetting::PadToPowerOfTwo || PowerOfTwoMode == ETexturePowerOfTwoSetting::PadToSquarePowerOfTwo)
 	{
diff --git a/Engine/Source/Runtime/Engine/Private/TextureDerivedData.cpp b/Engine/Source/Runtime/Engine/Private/TextureDerivedData.cpp
--- a/Engine/Source/Runtime/Engine/Private/TextureDerivedData.cpp
+++ b/Engine/Source/Runtime/Engine/Private/TextureDerivedData.cpp
@@ -39,6 +39,7 @@
 #include "VT/LightmapVirtualTexture.h"
 #include "TextureBuildUtilities.h"
 #include "TextureCompiler.h"
+#include "TextureCompressorModule.h"
 #include "TextureEncodingSettings.h"
 
 static TAutoConsoleVariable<int32> CVarTexturesCookToDerivedDataReferences(
@@ -1096,6 +1097,56 @@ static void GetTextureBuildSettings(
 	FinalizeBuildSettingsForLayer(Texture, 0, &TargetPlatform, InEncodeSpeed, OutBuildSettings, OutBuildResultMetadata);
 }
 
+/**
+ * Sets build settings for a texture on the target platform
+ * @param Texture - The texture for which to build compressor settings.
+ * @param OutBuildSettings - Array of desired texture settings
+ */
+static void GetBuildSettingsForTargetPlatform(
+	const UTexture& Texture,
+	const ITargetPlatform* TargetPlatform,
+	ETextureEncodeSpeed InEncodeSpeed, //  must be Fast or Final
+	TArray<FTextureBuildSettings>& OutSettingPerLayer,
+	TArray<FTexturePlatformData::FTextureEncodeResultMetadata>* OutResultMetadataPerLayer // can be nullptr if not needed
+)
+{
+	check(TargetPlatform != NULL);
+
+	const UTextureLODSettings* LODSettings = (UTextureLODSettings*)UDeviceProfileManager::Get().FindProfile(TargetPlatform->PlatformName());
+	FTextureBuildSettings SourceBuildSettings;
+	FTexturePlatformData::FTextureEncodeResultMetadata SourceMetadata;
+	GetTextureBuildSettings(Texture, *LODSettings, *TargetPlatform, InEncodeSpeed, SourceBuildSettings, &SourceMetadata);
+
+	TArray< TArray<FName> > PlatformFormats;
+	Texture.GetPlatformTextureFormatNamesWithPrefix(TargetPlatform,PlatformFormats);
+
+	// this code only uses PlatformFormats[0] , so it would be wrong for Android_Multi
+	//	but it's only used for the platform running the Editor
+	check(PlatformFormats.Num() == 1);
+
+	const int32 NumLayers = Texture.Source.GetNumLayers();
+	check(PlatformFormats[0].Num() == NumLayers);
+
+	OutSettingPerLayer.Reserve(NumLayers);
+	if (OutResultMetadataPerLayer)
+	{
+		OutResultMetadataPerLayer->Reserve(NumLayers);
+	}
+	for (int32 LayerIndex = 0; LayerIndex < NumLayers; ++LayerIndex)
+	{
+		FTextureBuildSettings& OutSettings = OutSettingPerLayer.Add_GetRef(SourceBuildSettings);
+		OutSettings.TextureFormatName = PlatformFormats[0][LayerIndex];
+
+		FTexturePlatformData::FTextureEncodeResultMetadata* OutMetadata = nullptr;
+		if (OutResultMetadataPerLayer)
+		{
+			OutMetadata = &OutResultMetadataPerLayer->Add_GetRef(SourceMetadata);
+		}
+			
+		FinalizeBuildSettingsForLayer(Texture, LayerIndex, TargetPlatform, InEncodeSpeed, OutSettings, OutMetadata);
+	}
+}
+
 /**
  * Sets build settings for a texture on the current running platform
  * @param Texture - The texture for which to build compressor settings.
@@ -1130,39 +1181,7 @@ static void GetBuildSettingsForRunningPlatform(
 
 		check(TargetPlatform != NULL);
 
-		const UTextureLODSettings* LODSettings = (UTextureLODSettings*)UDeviceProfileManager::Get().FindProfile(TargetPlatform->PlatformName());
-		FTextureBuildSettings SourceBuildSettings;
-		FTexturePlatformData::FTextureEncodeResultMetadata SourceMetadata;
-		GetTextureBuildSettings(Texture, *LODSettings, *TargetPlatform, InEncodeSpeed, SourceBuildSettings, &SourceMetadata);
-
-		TArray< TArray<FName> > PlatformFormats;
-		Texture.GetPlatformTextureFormatNamesWithPrefix(TargetPlatform,PlatformFormats);
-
-		// this code only uses PlatformFormats[0] , so it would be wrong for Android_Multi
-		//	but it's only used for the platform running the Editor
-		check(PlatformFormats.Num() == 1);
-
-		const int32 NumLayers = Texture.Source.GetNumLayers();
-		check(PlatformFormats[0].Num() == NumLayers);
-
-		OutSettingPerLayer.Reserve(NumLayers);
-		if (OutResultMetadataPerLayer)
-		{
-			OutResultMetadataPerLayer->Reserve(NumLayers);
-		}
-		for (int32 LayerIndex = 0; LayerIndex < NumLayers; ++LayerIndex)
-		{
-			FTextureBuildSettings& OutSettings = OutSettingPerLayer.Add_GetRef(SourceBuildSettings);
-			OutSettings.TextureFormatName = PlatformFormats[0][LayerIndex];
-
-			FTexturePlatformData::FTextureEncodeResultMetadata* OutMetadata = nullptr;
-			if (OutResultMetadataPerLayer)
-			{
-				OutMetadata = &OutResultMetadataPerLayer->Add_GetRef(SourceMetadata);
-			}
-			
-			FinalizeBuildSettingsForLayer(Texture, LayerIndex, TargetPlatform, InEncodeSpeed, OutSettings, OutMetadata);
-		}
+		GetBuildSettingsForTargetPlatform(Texture, TargetPlatform, InEncodeSpeed, OutSettingPerLayer, OutResultMetadataPerLayer);
 	}
 }
 
@@ -3833,3 +3852,70 @@ void UTexture::SetMinTextureResidentMipCount(int32 InMinTextureResidentMipCount)
 	int32 MinAllowedMipCount = FPlatformProperties::RequiresCookedData() ? 1 : NUM_INLINE_DERIVED_MIPS;
 	GMinTextureResidentMipCount = FMath::Max(InMinTextureResidentMipCount, MinAllowedMipCount);
 }
+
+#if WITH_EDITOR
+bool UTexture::DownsizeImageUsingTextureSettings(const ITargetPlatform* TargetPlatform, FImage& InOutImage, int32 TargetSize, int32 LayerIndex)
+{
+	if (TargetSize >= InOutImage.SizeX && TargetSize >= InOutImage.SizeY)
+	{
+		return false;
+	}
+
+	// Ideally this code wouldn't live here but at the moment of writing this code the coupling between the texture and the texture compressor make it hard to move that logic elsewhere
+	TArray<FTextureBuildSettings> SettingPerLayer;
+	TArray<FTexturePlatformData::FTextureEncodeResultMetadata>* ResultMetadataPerLayer = nullptr;
+	GetBuildSettingsForTargetPlatform(*this, TargetPlatform, ETextureEncodeSpeed::Final, SettingPerLayer, nullptr);
+
+	// Teak the build setting to generate a mip for our image
+	FTextureBuildSettings& BuildSettings = SettingPerLayer[LayerIndex];
+	BuildSettings.bCubemap = false;
+	BuildSettings.bTextureArray = false;
+	BuildSettings.bVolume = false;
+	BuildSettings.bLongLatSource = false;
+
+	FImage Temp;
+	// convert to RGBA32F linear for the compressor
+	InOutImage.CopyTo(Temp, ERawImageFormat::RGBA32F, EGammaSpace::Linear);
+
+	TArray<FImage> BuildSourceImageMips;
+	// make sure BuildSourceImageMips doesn't reallocate :
+	constexpr int BuildSourceImageMipsMaxCount = 20; // plenty
+	BuildSourceImageMips.Empty(BuildSourceImageMipsMaxCount);
+
+	ITextureCompressorModule::GenerateMipChain(BuildSettings, Temp, BuildSourceImageMips, 1);
+
+	while (BuildSourceImageMips.Last().SizeX > TargetSize ||
+		BuildSourceImageMips.Last().SizeY > TargetSize)
+	{
+		ITextureCompressorModule::GenerateMipChain(BuildSettings, BuildSourceImageMips.Last(), BuildSourceImageMips, 1);
+	}
+
+	FImage* SelectedOutput = nullptr;
+	if (BuildSourceImageMips.Last().SizeX < TargetSize ||
+		BuildSourceImageMips.Last().SizeY < TargetSize)
+	{
+		if (BuildSourceImageMips.Num() == 1)
+		{
+			return false;
+		}
+
+		SelectedOutput = &BuildSourceImageMips[BuildSourceImageMips.Num() - 2];
+	}
+	else
+	{
+		SelectedOutput = &BuildSourceImageMips.Last();
+	}
+
+
+	if (SelectedOutput->Format == InOutImage.Format && SelectedOutput->GammaSpace == InOutImage.GammaSpace)
+	{
+		InOutImage = MoveTemp(*SelectedOutput);
+	}
+	else
+	{
+		SelectedOutput->CopyTo(InOutImage, InOutImage.Format, InOutImage.GammaSpace);
+	}
+
+	return true;
+}
+#endif
diff --git a/Engine/Source/Runtime/TextureUtilitiesCommon/Private/TextureSourceDataUtils.cpp b/Engine/Source/Runtime/TextureUtilitiesCommon/Private/TextureSourceDataUtils.cpp
new file mode 100644
--- /dev/null
+++ b/Engine/Source/Runtime/TextureUtilitiesCommon/Private/TextureSourceDataUtils.cpp
@@ -0,0 +1,280 @@
+// Copyright Epic Games, Inc. All Rights Reserved.
+
+#include "TextureSourceDataUtils.h"
+
+#include "ImageCoreUtils.h"
+#include "Engine/Texture.h"
+#include "HAL/UnrealMemory.h"
+
+namespace UE::TextureUtilitiesCommon::Experimental
+{
+
+namespace Private
+{
+	bool ResizeTexture2D(UTexture* Texture, int32 MaxSize, const ITargetPlatform* TargetPlatform)
+	{
+		// Protect the code from an async build of the texture
+		Texture->PreEditChange(nullptr);
+
+		// We want to reduce the asset size so ignore the imported mip(s)
+		const int32 MipIndex = 0;
+		FImage SourceMip0;
+		if (!Texture->Source.GetMipImage(SourceMip0, MipIndex))
+		{
+			return false;
+		}
+
+		const int32 LayerIndex = 0;
+		Texture->DownsizeImageUsingTextureSettings(TargetPlatform, SourceMip0, MaxSize, LayerIndex);
+		FImage ResizedImage = MoveTemp(SourceMip0);
+
+		UE::Serialization::FEditorBulkData::FSharedBufferWithID ResizedImageBufferWithID = MakeSharedBufferFromArray(MoveTemp(ResizedImage.RawData));
+
+		const int32 NumMips = 1;
+		Texture->Source.Init(ResizedImage.SizeX
+			, ResizedImage.SizeY
+			, ResizedImage.NumSlices
+			, NumMips
+			, FImageCoreUtils::ConvertToTextureSourceFormat(ResizedImage.Format)
+			, MoveTemp(ResizedImageBufferWithID));
+
+		return true;
+	}
+
+	bool ResizeTextureSlicedBy2DLayers(UTexture* Texture, int32 MaxSize, const ITargetPlatform* TargetPlatform)
+	{
+		// Protect the code from an async build of the texture
+		Texture->PreEditChange(nullptr);
+
+		FImage ResizedImage;
+		{
+			TArray<FImage> ResizedSlices;
+			ResizedSlices.Reserve(Texture->Source.GetNumSlices());
+
+			ERawImageFormat::Type FormatUsed;
+			EGammaSpace GammaSpaceUsed;
+
+			{
+				// We want to reduce the asset size so ignore the imported mip(s)
+				TArray<FImage> Slices;
+				Slices.Reserve(Texture->Source.GetNumSlices());
+
+				{
+					// We could probably avoid a copy here but for the simplicity of the code keep it for now.
+					FImage SourceMip0;
+					if (!Texture->Source.GetMipImage(SourceMip0, 0))
+					{
+						return false;
+					}
+
+					FormatUsed = SourceMip0.Format;
+					GammaSpaceUsed = SourceMip0.GammaSpace;
+
+					for (int32 Index = 0; Index < SourceMip0.NumSlices; ++Index)
+					{
+						FImage& Slice = Slices.AddDefaulted_GetRef();
+						FImageView SliceView = SourceMip0.GetSlice(Index);
+
+						FImageCore::CopyImage(SliceView, Slice);
+					}
+				}
+
+				for (FImage& Slice : Slices)
+				{
+					FImage& ResizedSlice = ResizedSlices.AddDefaulted_GetRef();
+
+					const int32 LayerIndex = 0;
+					Texture->DownsizeImageUsingTextureSettings(TargetPlatform, Slice, MaxSize, LayerIndex);
+					ResizedSlice = MoveTemp(Slice);
+				}
+			}
+
+			// Move the resized slices into the resized image
+			ResizedImage.Init(ResizedSlices[0].SizeX,ResizedSlices[0].SizeY, ResizedSlices.Num(), FormatUsed, GammaSpaceUsed);
+
+			for (int32 Index = 0; Index < ResizedImage.NumSlices; ++Index)
+			{
+				FImageCore::CopyImage(ResizedSlices[Index], ResizedImage.GetSlice(Index));
+			}
+		}
+
+
+		UE::Serialization::FEditorBulkData::FSharedBufferWithID ResizedImageBufferWithID = MakeSharedBufferFromArray(MoveTemp(ResizedImage.RawData));
+
+		const int32 NumMips = 1;
+		Texture->Source.Init(ResizedImage.SizeX
+			, ResizedImage.SizeY
+			, ResizedImage.NumSlices
+			, NumMips
+			, FImageCoreUtils::ConvertToTextureSourceFormat(ResizedImage.Format)
+			, MoveTemp(ResizedImageBufferWithID));
+
+		return true;
+	}
+
+	bool ResizeTexture2DBlocked(UTexture* Texture, int32 MaxSize, const ITargetPlatform* TargetPlatform)
+	{
+		// Protect the code from an async build of the texture
+		Texture->PreEditChange(nullptr);
+
+		FIntPoint LogicalSourceSize = Texture->Source.GetLogicalSize();
+		double RatioX = double(MaxSize) / LogicalSourceSize.X;
+		double RatioY = double(MaxSize) / LogicalSourceSize.Y;
+
+		TArray<FTextureSourceBlock > ResizedSourceBlocks;
+		ResizedSourceBlocks.Reserve(Texture->Source.GetNumBlocks());
+	
+		TArray<FImage> ResizedBlocks;
+		ResizedBlocks.Reserve(Texture->Source.GetNumBlocks());
+
+		for (int32 BlockIndex = 0; BlockIndex < Texture->Source.GetNumBlocks(); ++BlockIndex)
+		{
+			// We want to reduce the asset size so ignore the imported mip(s)
+			FImage SourceMip0;
+			const int32 MipIndex = 0;
+			const int32 LayerIndex = 0;
+			if (!Texture->Source.GetMipImage(SourceMip0, BlockIndex, LayerIndex, MipIndex))
+			{
+				return false;
+			}
+
+			FTextureSourceBlock& ResizedSourceBlock = ResizedSourceBlocks.AddDefaulted_GetRef();
+			Texture->Source.GetBlock(BlockIndex, ResizedSourceBlock);
+		
+			FImage& ResizedBlock = ResizedBlocks.AddDefaulted_GetRef();
+
+			int32 BlockMaxSize = FMath::RoundToInt32(FMath::Min(ResizedSourceBlock.SizeX * RatioX, ResizedSourceBlock.SizeY * RatioY));
+
+			Texture->DownsizeImageUsingTextureSettings(TargetPlatform, SourceMip0, MaxSize, LayerIndex);
+			ResizedBlock = MoveTemp(SourceMip0);
+			
+			ResizedSourceBlock.SizeX = ResizedBlock.SizeX;
+			ResizedSourceBlock.SizeY = ResizedBlock.SizeY;
+			ResizedSourceBlock.NumSlices = 1;
+		}
+
+		int64 SizeNeededInBytes = 0;
+		for (const FImage& Block : ResizedBlocks)
+		{
+			SizeNeededInBytes += Block.RawData.Num();
+		}
+		FUniqueBuffer WriteImageBuffer = FUniqueBuffer::Alloc(SizeNeededInBytes);
+
+		uint8* CurrentAddress = static_cast<uint8*>(WriteImageBuffer.GetData());
+		for (FImage& Block : ResizedBlocks)
+		{
+			FMemory::Memcpy(CurrentAddress, static_cast<uint8*>(Block.RawData.GetData()), Block.RawData.Num());
+			CurrentAddress += Block.RawData.Num();
+		}
+
+		UE::Serialization::FEditorBulkData::FSharedBufferWithID ResizedImageBufferWithID = WriteImageBuffer.MoveToShared();
+
+		const ETextureSourceFormat SourceFormat = Texture->Source.GetFormat();
+		int32 NumLayers = 1;
+		Texture->Source.InitBlocked(
+			&SourceFormat,
+			ResizedSourceBlocks.GetData(),
+			NumLayers,
+			ResizedSourceBlocks.Num(),
+			MoveTemp(ResizedImageBufferWithID)
+		);
+
+		return true;
+	}
+}
+
+
+bool DownsizeTextureSourceData(UTexture* Texture, int32 TargetSizeInGame, const ITargetPlatform* TargetPlatform)
+{
+	// Check if we don't know how to resize that texture
+	if (Texture->Source.GetNumLayers() != 1)
+	{
+		return false;
+	}
+
+	if (!(Texture->Source.GetTextureClass() == ETextureClass::Cube || Texture->Source.GetTextureClass() == ETextureClass::TwoD))
+	{
+		return false;
+	}
+
+	if (Texture->Source.GetTextureClass() == ETextureClass::TwoD && Texture->Source.GetNumSlices() != 1)
+	{
+		return false;
+	}
+
+	if (Texture->Source.GetTextureClass() == ETextureClass::Cube)
+	{
+		if (Texture->Source.IsLongLatCubemap())
+		{
+			if (Texture->Source.GetNumSlices() != 1)
+			{
+				return false;
+			}
+		}
+		else if (Texture->Source.GetNumSlices() != 6)
+		{
+			return false;
+		}
+	}
+
+	int32 TargetSourceSize = TargetSizeInGame;
+	if (Texture->Source.IsLongLatCubemap())
+	{
+		// The function return the max size of the generated cube from the source long lat
+		// this should be kept in sync with the implementation details of ComputeLongLatCubemapExtents() or refactored
+		TargetSourceSize = (1U << FMath::FloorLog2(TargetSizeInGame)) * 2;
+	}
+
+	FIntPoint SourceSize = Texture->Source.GetLogicalSize();
+	if (SourceSize.X <= TargetSourceSize && SourceSize.Y <= TargetSourceSize)
+	{
+		return false;
+	}
+
+	if (Texture->Source.GetTextureClass() == ETextureClass::TwoD)
+	{
+		if (Texture->Source.GetNumBlocks() == 1)
+		{
+			return Private::ResizeTexture2D(Texture, TargetSourceSize, TargetPlatform);
+		}
+		else
+		{
+			// UDIM(s)
+			return Private::ResizeTexture2DBlocked(Texture, TargetSourceSize, TargetPlatform);
+		}
+	}
+	else if (Texture->Source.GetTextureClass() == ETextureClass::Cube)
+	{
+		if (Texture->Source.IsLongLatCubemap())
+		{
+			return Private::ResizeTexture2D(Texture, TargetSourceSize, TargetPlatform);
+		}
+		else
+		{
+			return Private::ResizeTextureSlicedBy2DLayers(Texture, TargetSourceSize, TargetPlatform);
+		}
+	}
+	
+	return false;
+}
+
+bool DownsizeTexureSourceDataNearRenderingSize(UTexture* Texture, const ITargetPlatform* TargetPlatform)
+{
+	int32 SizeX;
+	int32 SizeY;
+	Texture->GetBuiltTextureSize(TargetPlatform, SizeX, SizeY);
+
+	int32 TargetSize =  FMath::Max(SizeX, SizeY);
+	if (DownsizeTextureSourceData(Texture, FMath::Max(SizeX, SizeY), TargetPlatform))
+	{
+		int32 BiggestSize = FMath::Max(SizeX, SizeY);
+		const int32 Bias = FMath::CeilLogTwo(BiggestSize) - FMath::CeilLogTwo(TargetSize);
+		Texture->LODBias = FMath::Max(Bias , 0);
+		Texture->PostEditChange();
+		return true;
+	}
+
+	return false;
+}
+
+} // End namespace UE::TextureUtilitiesCommon::Experimental
diff --git a/Engine/Source/Runtime/TextureUtilitiesCommon/Public/TextureSourceDataUtils.h b/Engine/Source/Runtime/TextureUtilitiesCommon/Public/TextureSourceDataUtils.h
new file mode 100644
--- /dev/null
+++ b/Engine/Source/Runtime/TextureUtilitiesCommon/Public/TextureSourceDataUtils.h
@@ -0,0 +1,22 @@
+// Copyright Epic Games, Inc. All Rights Reserved.
+
+#pragma once
+
+#if WITH_EDITOR
+#include "CoreTypes.h"
+
+class UTexture;
+class ITargetPlatform;
+
+// Not tested against all the texture types and the api might change or be removed without deprecation
+namespace UE::TextureUtilitiesCommon::Experimental
+{
+	// Try to resize the texture source data so that it is not larger than one build for the specified platform (only Texture2D and TextureCube are supported at the moment)
+	// Note it does not trigger the post edit change after modifying the source texture
+	TEXTUREUTILITIESCOMMON_API bool DownsizeTextureSourceData(UTexture* Texture, int32 TargetSizeInGame, const ITargetPlatform* TargetPlatform);
+
+	//Note: This function does trigger the post edit change after modifying the source texture
+	TEXTUREUTILITIESCOMMON_API bool DownsizeTexureSourceDataNearRenderingSize(UTexture* Texture, const ITargetPlatform* TargetPlatform);
+}
+
+#endif // WITH_EDITOR
\ No newline at end of file
