Fixed a build issue at runtime. The code in the file is only for the editor but it wasn't wrapped in a if WITH_EDITOR.

#rnx

#localization none
#tests Just submitting the fix. So no other testing then just checking if it compile was done.

diff --git a/Engine/Source/Runtime/TextureUtilitiesCommon/Private/TextureSourceDataUtils.cpp b/Engine/Source/Runtime/TextureUtilitiesCommon/Private/TextureSourceDataUtils.cpp
--- a/Engine/Source/Runtime/TextureUtilitiesCommon/Private/TextureSourceDataUtils.cpp
+++ b/Engine/Source/Runtime/TextureUtilitiesCommon/Private/TextureSourceDataUtils.cpp
@@ -2,6 +2,8 @@
 
 #include "TextureSourceDataUtils.h"
 
+#if WITH_EDITOR
+
 #include "ImageCoreUtils.h"
 #include "Engine/Texture.h"
 #include "HAL/UnrealMemory.h"
@@ -278,3 +280,5 @@ bool DownsizeTexureSourceDataNearRenderingSize(UTexture* Texture, const ITargetP
 }
 
 } // End namespace UE::TextureUtilitiesCommon::Experimental
+
+#endif //WITH_EDITOR
