fix some issues with Resize Import Texture
ensure texture has valid source
must change pow22 gamma to srgb
was not correct resizing non-square textures
add verification of correct built size

diff --git a/Engine/Source/Runtime/Engine/Private/TextureDerivedData.cpp b/Engine/Source/Runtime/Engine/Private/TextureDerivedData.cpp
--- a/Engine/Source/Runtime/Engine/Private/TextureDerivedData.cpp
+++ b/Engine/Source/Runtime/Engine/Private/TextureDerivedData.cpp
@@ -3756,14 +3756,16 @@ void UTexture::SetMinTextureResidentMipCount(int32 InMinTextureResidentMipCount)
 #if WITH_EDITOR
 bool UTexture::DownsizeImageUsingTextureSettings(const ITargetPlatform* TargetPlatform, FImage& InOutImage, int32 TargetSize, int32 LayerIndex)
 {
+	// resize so that the largest dimension is <= TargetSize
+
 	if (TargetSize >= InOutImage.SizeX && TargetSize >= InOutImage.SizeY)
 	{
+		// both dimensions already small enough, early out
 		return false;
 	}
 
 	// Ideally this code wouldn't live here but at the moment of writing this code the coupling between the texture and the texture compressor make it hard to move that logic elsewhere
 	TArray<FTextureBuildSettings> SettingPerLayer;
-	TArray<FTexturePlatformData::FTextureEncodeResultMetadata>* ResultMetadataPerLayer = nullptr;
 	GetBuildSettingsForTargetPlatform(*this, TargetPlatform, ETextureEncodeSpeed::Final, SettingPerLayer, nullptr);
 
 	// Teak the build setting to generate a mip for our image
@@ -3777,6 +3779,13 @@ bool UTexture::DownsizeImageUsingTextureSettings(const ITargetPlatform* TargetPl
 	// convert to RGBA32F linear for the compressor
 	InOutImage.CopyTo(Temp, ERawImageFormat::RGBA32F, EGammaSpace::Linear);
 
+	if ( InOutImage.GetGammaSpace() == EGammaSpace::Pow22 )
+	{
+		// Pow22 is read only; now that we have converted to Linear, we will write output as sRGB
+		// caller must also change Texture->bUseLegacyGamma
+		InOutImage.GammaSpace = EGammaSpace::sRGB;
+	}
+
 	TArray<FImage> BuildSourceImageMips;
 	// make sure BuildSourceImageMips doesn't reallocate :
 	constexpr int BuildSourceImageMipsMaxCount = 20; // plenty
@@ -3784,28 +3793,16 @@ bool UTexture::DownsizeImageUsingTextureSettings(const ITargetPlatform* TargetPl
 
 	ITextureCompressorModule::GenerateMipChain(BuildSettings, Temp, BuildSourceImageMips, 1);
 
+	// keep halving while larger size is > Target
 	while (BuildSourceImageMips.Last().SizeX > TargetSize ||
 		BuildSourceImageMips.Last().SizeY > TargetSize)
 	{
 		ITextureCompressorModule::GenerateMipChain(BuildSettings, BuildSourceImageMips.Last(), BuildSourceImageMips, 1);
 	}
 
-	FImage* SelectedOutput = nullptr;
-	if (BuildSourceImageMips.Last().SizeX < TargetSize ||
-		BuildSourceImageMips.Last().SizeY < TargetSize)
-	{
-		if (BuildSourceImageMips.Num() == 1)
-		{
-			return false;
-		}
-
-		SelectedOutput = &BuildSourceImageMips[BuildSourceImageMips.Num() - 2];
-	}
-	else
-	{
-		SelectedOutput = &BuildSourceImageMips.Last();
-	}
+	// now larger size must be <= TargetSize
 
+	FImage* SelectedOutput = &BuildSourceImageMips.Last();
 
 	if (SelectedOutput->Format == InOutImage.Format && SelectedOutput->GammaSpace == InOutImage.GammaSpace)
 	{
diff --git a/Engine/Source/Runtime/TextureUtilitiesCommon/Private/TextureSourceDataUtils.cpp b/Engine/Source/Runtime/TextureUtilitiesCommon/Private/TextureSourceDataUtils.cpp
--- a/Engine/Source/Runtime/TextureUtilitiesCommon/Private/TextureSourceDataUtils.cpp
+++ b/Engine/Source/Runtime/TextureUtilitiesCommon/Private/TextureSourceDataUtils.cpp
@@ -7,6 +7,7 @@
 #include "ImageCoreUtils.h"
 #include "Engine/Texture.h"
 #include "HAL/UnrealMemory.h"
+#include "EngineLogs.h"
 
 namespace UE::TextureUtilitiesCommon::Experimental
 {
@@ -27,7 +28,10 @@ namespace Private
 		}
 
 		const int32 LayerIndex = 0;
-		Texture->DownsizeImageUsingTextureSettings(TargetPlatform, SourceMip0, MaxSize, LayerIndex);
+		if ( ! Texture->DownsizeImageUsingTextureSettings(TargetPlatform, SourceMip0, MaxSize, LayerIndex) )
+		{
+			return false;
+		}
 		FImage ResizedImage = MoveTemp(SourceMip0);
 
 		UE::Serialization::FEditorBulkData::FSharedBufferWithID ResizedImageBufferWithID = MakeSharedBufferFromArray(MoveTemp(ResizedImage.RawData));
@@ -40,6 +44,11 @@ namespace Private
 			, FImageCoreUtils::ConvertToTextureSourceFormat(ResizedImage.Format)
 			, MoveTemp(ResizedImageBufferWithID));
 
+		// if gamma was Pow22 it is now sRGB
+		Texture->bUseLegacyGamma = false;
+
+		// PostEditChange is called outside by our caller
+
 		return true;
 	}
 
@@ -188,6 +197,8 @@ namespace Private
 
 bool DownsizeTextureSourceData(UTexture* Texture, int32 TargetSizeInGame, const ITargetPlatform* TargetPlatform)
 {
+	check( Texture->Source.IsValid() );
+
 	// Check if we don't know how to resize that texture
 	if (Texture->Source.GetNumLayers() != 1)
 	{
@@ -262,17 +273,35 @@ bool DownsizeTextureSourceData(UTexture* Texture, int32 TargetSizeInGame, const
 
 bool DownsizeTexureSourceDataNearRenderingSize(UTexture* Texture, const ITargetPlatform* TargetPlatform)
 {
-	int32 SizeX;
-	int32 SizeY;
-	Texture->GetBuiltTextureSize(TargetPlatform, SizeX, SizeY);
+	if ( ! Texture->Source.IsValid() )
+	{
+		return false;
+	}
+
+	int32 BeforeSizeX;
+	int32 BeforeSizeY;
+	Texture->GetBuiltTextureSize(TargetPlatform, BeforeSizeX, BeforeSizeY);
 
-	int32 TargetSize =  FMath::Max(SizeX, SizeY);
-	if (DownsizeTextureSourceData(Texture, FMath::Max(SizeX, SizeY), TargetPlatform))
+	int32 TargetSize = FMath::Max(BeforeSizeX, BeforeSizeY);
+	if (DownsizeTextureSourceData(Texture, TargetSize, TargetPlatform))
 	{
-		int32 BiggestSize = FMath::Max(SizeX, SizeY);
-		const int32 Bias = FMath::CeilLogTwo(BiggestSize) - FMath::CeilLogTwo(TargetSize);
-		Texture->LODBias = FMath::Max(Bias , 0);
+		Texture->LODBias = 0;
 		Texture->PostEditChange();
+		
+		// check that GetBuiltTextureSize was preserved :
+		int32 AfterSizeX;
+		int32 AfterSizeY;
+		Texture->GetBuiltTextureSize(TargetPlatform, AfterSizeX, AfterSizeY);
+
+		if ( BeforeSizeX != AfterSizeX ||
+			 BeforeSizeY != AfterSizeY )
+		{
+			UE_LOG(LogTexture,Warning,TEXT("DownsizeTexureSourceDataNearRenderingSize failed to preserve built size; was: %dx%d now: %dx%d on [%s]"),
+				BeforeSizeX,BeforeSizeY,
+				AfterSizeX,AfterSizeY,
+				*Texture->GetFullName());
+		}
+
 		return true;
 	}
 
