ETextureImportPNGInfill option to control when transparent RGB is filled from neighbors
set to Always to get fill into transparent for non-binary transparency case
legacy config setting still respected, but new config supercedes
legacy behavior was to Infill only for binary transparency

#rnx

diff --git a/Engine/Plugins/Interchange/Runtime/Source/Import/Private/Texture/InterchangeImageWrapperTranslator.cpp b/Engine/Plugins/Interchange/Runtime/Source/Import/Private/Texture/InterchangeImageWrapperTranslator.cpp
--- a/Engine/Plugins/Interchange/Runtime/Source/Import/Private/Texture/InterchangeImageWrapperTranslator.cpp
+++ b/Engine/Plugins/Interchange/Runtime/Source/Import/Private/Texture/InterchangeImageWrapperTranslator.cpp
@@ -68,12 +68,12 @@ static FAutoConsoleVariableRef CCvarInterchangeEnableTGAImport(
 
 UInterchangeImageWrapperTranslator::UInterchangeImageWrapperTranslator()
 {
-	if (IsTemplate() && GConfig)
-	{
-		ensure(IsInGameThread()); // Reading config values isn't threadsafe
-		bFillPNGZeroAlpha = true;
-		GConfig->GetBool(TEXT("TextureImporter"), TEXT("FillPNGZeroAlpha"), bFillPNGZeroAlpha.GetValue(), GEditorIni);
-	}
+	// construction of CDO is not thread safe,
+	//  so ensure it is done on game thread before using it from threads
+	// that is guaranteed because the module loader inits all CDOs
+
+	PNGInfill = GetDefault<UTextureImportSettings>()->GetPNGInfillMapDefault();
+	check( PNGInfill != ETextureImportPNGInfill::Default );
 }
 
 TArray<FString> UInterchangeImageWrapperTranslator::GetSupportedFormats() const
@@ -193,10 +193,12 @@ TOptional<UE::Interchange::FImportImage> UInterchangeImageWrapperTranslator::Get
 
 			if (ImageFormat == EImageFormat::PNG)
 			{
-				if (GetDefault<UInterchangeImageWrapperTranslator>()->bFillPNGZeroAlpha.GetValue())
+				if (PNGInfill != ETextureImportPNGInfill::Never)
 				{
+					bool bDoOnComplexAlphaNotJustBinaryTransparency = ( PNGInfill == ETextureImportPNGInfill::Always );
+
 					// Replace the pixels with 0.0 alpha with a color value from the nearest neighboring color which has a non-zero alpha
-					UE::TextureUtilitiesCommon::FillZeroAlphaPNGData(PayloadData.SizeX, PayloadData.SizeY, PayloadData.Format, reinterpret_cast<uint8*>(PayloadData.RawData.GetData()));
+					UE::TextureUtilitiesCommon::FillZeroAlphaPNGData(PayloadData.SizeX, PayloadData.SizeY, PayloadData.Format, reinterpret_cast<uint8*>(PayloadData.RawData.GetData()), bDoOnComplexAlphaNotJustBinaryTransparency);
 				}
 			}
 			else if (ImageFormat == EImageFormat::TGA)
diff --git a/Engine/Plugins/Interchange/Runtime/Source/Import/Public/Texture/InterchangeImageWrapperTranslator.h b/Engine/Plugins/Interchange/Runtime/Source/Import/Public/Texture/InterchangeImageWrapperTranslator.h
--- a/Engine/Plugins/Interchange/Runtime/Source/Import/Public/Texture/InterchangeImageWrapperTranslator.h
+++ b/Engine/Plugins/Interchange/Runtime/Source/Import/Public/Texture/InterchangeImageWrapperTranslator.h
@@ -9,6 +9,7 @@
 #include "Nodes/InterchangeBaseNodeContainer.h"
 #include "Texture/InterchangeTexturePayloadData.h"
 #include "Texture/InterchangeTexturePayloadInterface.h"
+#include "TextureImportSettings.h"
 
 #include "InterchangeImageWrapperTranslator.generated.h"
 
@@ -51,5 +52,5 @@ public:
 	TOptional<UE::Interchange::FImportImage> GetTexturePayloadDataFromBuffer(const TArray64<uint8>& SourceDataBuffer) const;
 
 private:
-	TOptional<bool> bFillPNGZeroAlpha; // Set on the CDO only
+	ETextureImportPNGInfill PNGInfill;
 };
\ No newline at end of file
diff --git a/Engine/Source/Editor/UnrealEd/Private/Factories/EditorFactories.cpp b/Engine/Source/Editor/UnrealEd/Private/Factories/EditorFactories.cpp
--- a/Engine/Source/Editor/UnrealEd/Private/Factories/EditorFactories.cpp
+++ b/Engine/Source/Editor/UnrealEd/Private/Factories/EditorFactories.cpp
@@ -2993,13 +2993,14 @@ bool UTextureFactory::ImportImage(const uint8* Buffer, int64 Length, FFeedbackCo
 
 			if ( ImageFormat == EImageFormat::PNG )
 			{
-				bool bFillPNGZeroAlpha = true;
-				GConfig->GetBool(TEXT("TextureImporter"), TEXT("FillPNGZeroAlpha"), bFillPNGZeroAlpha, GEditorIni);
+				ETextureImportPNGInfill PNGInfill = GetDefault<UTextureImportSettings>()->GetPNGInfillMapDefault();
 
-				if (bFillPNGZeroAlpha)
+				if (PNGInfill != ETextureImportPNGInfill::Never)
 				{
+					bool bDoOnComplexAlphaNotJustBinaryTransparency = ( PNGInfill == ETextureImportPNGInfill::Always );
+
 					// Replace the pixels with 0.0 alpha with a color value from the nearest neighboring color which has a non-zero alpha
-					UE::TextureUtilitiesCommon::FillZeroAlphaPNGData(OutImage.SizeX, OutImage.SizeY, OutImage.Format, OutImage.RawData.GetData());
+					UE::TextureUtilitiesCommon::FillZeroAlphaPNGData(OutImage.SizeX, OutImage.SizeY, OutImage.Format, OutImage.RawData.GetData(), bDoOnComplexAlphaNotJustBinaryTransparency);
 				}
 			}
 			else if ( ImageFormat == EImageFormat::TGA )
diff --git a/Engine/Source/Runtime/TextureUtilitiesCommon/Private/TextureImportSettings.cpp b/Engine/Source/Runtime/TextureUtilitiesCommon/Private/TextureImportSettings.cpp
--- a/Engine/Source/Runtime/TextureUtilitiesCommon/Private/TextureImportSettings.cpp
+++ b/Engine/Source/Runtime/TextureUtilitiesCommon/Private/TextureImportSettings.cpp
@@ -7,6 +7,7 @@
 
 #if WITH_EDITOR
 #include "Engine/Texture.h"
+#include "Misc/ConfigCacheIni.h"
 #endif
 
 #include UE_INLINE_GENERATED_CPP_BY_NAME(TextureImportSettings)
@@ -29,6 +30,29 @@ void UTextureImportSettings::PostInitProperties()
 }
 
 #if WITH_EDITOR
+
+// Get the PNGInfill setting, with Default mapped to a concrete choice
+ETextureImportPNGInfill UTextureImportSettings::GetPNGInfillMapDefault() const
+{
+	if ( PNGInfill == ETextureImportPNGInfill::Default )
+	{
+		// Default is OnlyOnBinaryTransparency unless changed by legacy config
+
+		// get legacy config :
+		bool bFillPNGZeroAlpha = true;
+		if ( GConfig )
+		{
+			GConfig->GetBool(TEXT("TextureImporter"), TEXT("FillPNGZeroAlpha"), bFillPNGZeroAlpha, GEditorIni);		
+		}
+		
+		return bFillPNGZeroAlpha ? ETextureImportPNGInfill::OnlyOnBinaryTransparency : ETextureImportPNGInfill::Never;
+	}
+	else
+	{
+		return PNGInfill;
+	}
+}
+
 void UTextureImportSettings::PostEditChangeProperty(FPropertyChangedEvent& PropertyChangedEvent)
 {
 	Super::PostEditChangeProperty(PropertyChangedEvent);
diff --git a/Engine/Source/Runtime/TextureUtilitiesCommon/Private/TextureImportUtils.cpp b/Engine/Source/Runtime/TextureUtilitiesCommon/Private/TextureImportUtils.cpp
--- a/Engine/Source/Runtime/TextureUtilitiesCommon/Private/TextureImportUtils.cpp
+++ b/Engine/Source/Runtime/TextureUtilitiesCommon/Private/TextureImportUtils.cpp
@@ -50,7 +50,7 @@ namespace UE
 		   libpng decodes those pels are {RGB=white, A=0}
 		   we replace them by filling in the RGB from neighbors
 
-		   note that this does NOT fill in the RGB of PNGs with a full alpha channel.
+		   note that this does NOT fill in the RGB of PNGs with a full alpha channel. -> it does now, if PNGInfill == Always
 		 */
 		template<typename PixelDataType, typename ColorDataType, int32 RIdx, int32 GIdx, int32 BIdx, int32 AIdx>
 		class TPNGDataFill
@@ -86,7 +86,7 @@ namespace UE
 						int rsqr = r1*r1 + r2*r2;
 						if( rsqr == 0 )
 						{
-							NearNeighborWeights[0][0] = 0.f;
+							NearNeighborWeights[0][0] = 0.f; // center self-weight = zero
 							continue;
 						}
 						float W = expf( - 1.1f * sqrtf((float)rsqr) );
@@ -95,10 +95,11 @@ namespace UE
 				}
 			}
 
-			void ProcessData()
+			void ProcessData(bool bDoOnComplexAlphaNotJustBinaryTransparency)
 			{
 				// first identify alpha type :
-				bool HasBinaryTransparency=false;
+				bool HasWhiteWithZeroAlpha=false;
+				bool HasComplexAlpha=false;
 				for (int64 Y = 0; Y < TextureHeight; ++Y)
 				{
 					const ColorDataType* RowData = (const ColorDataType *)SourceData + Y * TextureWidth;
@@ -110,25 +111,36 @@ namespace UE
 						}
 						else if ( RowData[X] == WhiteWithZeroAlpha )
 						{
-							HasBinaryTransparency = true;
+							HasWhiteWithZeroAlpha = true;
 						}
 						else
 						{
-							//HasComplexAlphaChannel = true;
-							UE_LOG(LogCore, Log,  TEXT("PNG has complex alpha channel, will not fill RGB in transparent background"));
-							// do not modify png's with full alpha channels :
-							return;
+							HasComplexAlpha = true;
+
+							if ( ! bDoOnComplexAlphaNotJustBinaryTransparency )
+							{
+								UE_LOG(LogCore, Log,  TEXT("PNG has complex alpha channel, will not fill RGB in transparent background, due to setting PNGInfill == OnlyOnBinaryTransparency"));
+								// do not modify png's with full alpha channels :
+								return;
+							}
 						}
 					}
 				}
 
-				if ( ! HasBinaryTransparency )
+				if ( ! HasWhiteWithZeroAlpha )
 				{
 					// all opaque
 					return;
 				}
 
-				UE_LOG(LogCore, Log,  TEXT("PNG has binary transparency, doing fill of RGB in transparent background"));
+				if ( HasComplexAlpha )
+				{
+					UE_LOG(LogCore, Log,  TEXT("PNG has alpha channel, doing fill of RGB in transparent background, due to setting PNGInfill == Always"));
+				}
+				else
+				{
+					UE_LOG(LogCore, Log,  TEXT("PNG has binary transparency, doing fill of RGB in transparent background, due to setting PNGInfill != Never"));
+				}
 
 				// first do good fill with limited distance :
 				// this ensures near pels within NearNeighborRadius get a good neighbor fill for interpolation
@@ -170,23 +182,6 @@ namespace UE
 						FillRowColorPixels(NumZeroedTopRowsToProcess, Y);
 					}
 				}
-
-				// change all pixels we made in FillFromNearNeighbors and change to A=0 :
-				
-				for (int64 Y = 0; Y < TextureHeight; ++Y)
-				{
-					ColorDataType * ImageRow = (ColorDataType *)SourceData + Y * TextureWidth;
-					for (int64 X = 0; X < TextureWidth; ++X)
-					{
-						const ColorDataType Color = ImageRow[X];
-						if ( Color == WhiteWithZeroAlpha || IsOpaque(Color) )
-						{
-							continue;
-						}
-
-						ImageRow[X] = MakeColorWithZeroAlpha(Color);
-					}
-				}
 			}
 
 			static bool IsOpaque(const ColorDataType InColor)
@@ -252,10 +247,15 @@ namespace UE
 				ColorDataType* RowData = (ColorDataType *)SourceData + Y * TextureWidth;
 				int64 X = 0;
 
+				// note this is done after the NN fill
+				// the NN fill will have RGB != white but A = 0
+				//	so we will fill out using those
+
 				if ( RowData[0] == WhiteWithZeroAlpha )
 				{
 					// transparent run at start of row
 					// find X which is the first opaque pel
+					//	( "opaque" is a misnomer; actually transparent but not WhiteWithZeroAlpha )
 
 					for(;;)
 					{
@@ -438,7 +438,8 @@ namespace UE
 					{
 						const ColorDataType Color = ImageColors[ X + Y * TextureWidth ];
 
-						// check IsOpaque rather than == WhiteWithZeroAlpha so that we can work in-place
+						#if 0
+						// only fill from opaque colors
 						if ( IsOpaque(Color) )
 						{
 							float W = RowNearNeighborWeights[ FMath::Abs(X-CenterX) ];
@@ -446,6 +447,23 @@ namespace UE
 							TotalWeight += W;
 							TotalColor += W * MakeLinearColor(Color);
 						}
+						#else
+						// fill from any non-transparent colors
+						// weight by alpha
+						FLinearColor CurColor = MakeLinearColor(Color);
+						if ( CurColor.A != 0.f )
+						{
+							float W = RowNearNeighborWeights[ FMath::Abs(X-CenterX) ];
+
+							// weighted by alpha :
+							W *= CurColor.A;
+
+							// accumulate weighted color
+							TotalWeight += W;
+							TotalColor += W * CurColor;
+						}
+						#endif
+
 					}
 				}
 
@@ -456,10 +474,10 @@ namespace UE
 				}
 
 				TotalColor *= 1.f / TotalWeight;
-				TotalColor.A = 0.5f; // anything other than opaque or zero
+				TotalColor.A = 0.f;
 
-				// we return an A that's not opaque so we won't be read from in this pass
-				// later we will change this A to zero
+				// returned Color has A = 0 so it will not be read from in this pass
+				//	if it has RGB != white it will not be filled again
 
 				return MakeColorFromLinear(TotalColor);
 			}
@@ -498,6 +516,8 @@ namespace UE
 							{
 								// could write to ImageRow[X] immediately, but using ScratchRow reduces cache sharing across cores
 								ScratchRow[X] = GetFilledFromNearNeighbors(X,Y);
+
+								// ScratchRow[X] still has zero alpha, but no longer white
 							}
 							else
 							{
@@ -519,21 +539,21 @@ namespace UE
 			float NearNeighborWeights[NearNeighborRadius+1][NearNeighborRadius+1];
 		};
 
-		void FillZeroAlphaPNGData(int32 SizeX, int32 SizeY, ETextureSourceFormat SourceFormat, uint8* SourceData)
+		void FillZeroAlphaPNGData(int32 SizeX, int32 SizeY, ETextureSourceFormat SourceFormat, uint8* SourceData, bool bDoOnComplexAlphaNotJustBinaryTransparency)
 		{
 			switch (SourceFormat)
 			{
 				case TSF_BGRA8:
 				{
 					TPNGDataFill<uint8, uint32, 2, 1, 0, 3> PNGFill(SizeX, SizeY, SourceData);
-					PNGFill.ProcessData();
+					PNGFill.ProcessData(bDoOnComplexAlphaNotJustBinaryTransparency);
 					break;
 				}
 
 				case TSF_RGBA16:
 				{
 					TPNGDataFill<uint16, uint64, 0, 1, 2, 3> PNGFill(SizeX, SizeY, SourceData);
-					PNGFill.ProcessData();
+					PNGFill.ProcessData(bDoOnComplexAlphaNotJustBinaryTransparency);
 					break;
 				}
 
diff --git a/Engine/Source/Runtime/TextureUtilitiesCommon/Public/TextureImportSettings.h b/Engine/Source/Runtime/TextureUtilitiesCommon/Public/TextureImportSettings.h
--- a/Engine/Source/Runtime/TextureUtilitiesCommon/Public/TextureImportSettings.h
+++ b/Engine/Source/Runtime/TextureUtilitiesCommon/Public/TextureImportSettings.h
@@ -26,6 +26,26 @@ enum class ETextureImportFloatingPointFormat : uint8
 	PreviousDefault = HDR_F16  UMETA(Hidden) // legacy behavior
 };
 
+
+/* When should RGB colors be spread into neighboring fully transparent white pixels, replacing their RGB.
+   By default, this is done OnlyOnBinaryTransparency, not on PNG's with non-binary-transparency alpha channels.
+   The PNG format has two different ways of storing alpha, either as 1-bit binary transparency, or as full 8/16 bit alpha channels.
+
+   Used to be set from the TextureImporter/FillPNGZeroAlpha config value.  Setting this option will supercede that.
+ */
+UENUM()
+enum class ETextureImportPNGInfill : uint8
+{
+	/* Use the legacy default behavior, set from the TextureImporter/FillPNGZeroAlpha config value; default was OnlyOnBinaryTransparency. */
+	Default = 0,
+	/* Never infill RGB, import the PNG exactly as it is stored in the file. */
+	Never,
+	/* Only infill RGB on binary transparency; this is the default behavior. */
+	OnlyOnBinaryTransparency,
+	/* Always infill RGB to fully transparent white pixels, even for non-binary alpha channels. */
+	Always
+};
+
 UCLASS(config=Editor, defaultconfig, meta=(DisplayName="Texture Import"))
 class TEXTUREUTILITIESCOMMON_API UTextureImportSettings : public UDeveloperSettings
 {
@@ -52,12 +72,20 @@ public:
 		DisplayName = "CompressionFormat to use for new float textures",
 		ToolTip = "Optionally use HDRCompressed (BC6H), or 32-bit adaptively, instead of HDR (RGBA16F) for floating point textures.  This setting is applied to newly imported textures, it does not affect existing textures in the project."))
 	ETextureImportFloatingPointFormat CompressedFormatForFloatTextures = ETextureImportFloatingPointFormat::PreviousDefault;
+	
+	UPROPERTY(config, EditAnywhere, Category=ImportSettings, meta = (
+		DisplayName = "When to infill RGB in transparent white PNG",
+		ToolTip = "Default behavior is to infill only for binary transparency; this setting may change that to always or never.  Will check TextureImporter/FillPNGZeroAlpha if this is not changed from Default.  This setting is applied to newly imported textures, it does not affect existing textures in the project."))
+	ETextureImportPNGInfill PNGInfill = ETextureImportPNGInfill::Default;
 
 	//~ Begin UObject Interface
 
 	virtual void PostInitProperties() override;
 
 #if WITH_EDITOR
+	// Get the PNGInfill setting, with Default mapped to a concrete choice
+	ETextureImportPNGInfill GetPNGInfillMapDefault() const;
+
 	virtual void PostEditChangeProperty(FPropertyChangedEvent& PropertyChangedEvent) override;
 #endif
 
diff --git a/Engine/Source/Runtime/TextureUtilitiesCommon/Public/TextureImportUtils.h b/Engine/Source/Runtime/TextureUtilitiesCommon/Public/TextureImportUtils.h
--- a/Engine/Source/Runtime/TextureUtilitiesCommon/Public/TextureImportUtils.h
+++ b/Engine/Source/Runtime/TextureUtilitiesCommon/Public/TextureImportUtils.h
@@ -25,5 +25,5 @@ namespace UE::TextureUtilitiesCommon
 	 * @param TextureSource - The source texture
 	 * @param SourceData - The source texture data
 	 */
-	TEXTUREUTILITIESCOMMON_API void FillZeroAlphaPNGData(int32 SizeX, int32 SizeY, ETextureSourceFormat SourceFormat, uint8* SourceData);
+	TEXTUREUTILITIESCOMMON_API void FillZeroAlphaPNGData(int32 SizeX, int32 SizeY, ETextureSourceFormat SourceFormat, uint8* SourceData, bool bDoOnComplexAlphaNotJustBinaryTransparency);
 }
\ No newline at end of file
