PNGDataFill bounds check fix

diff --git a/Engine/Source/Runtime/TextureUtilitiesCommon/Private/TextureImportUtils.cpp b/Engine/Source/Runtime/TextureUtilitiesCommon/Private/TextureImportUtils.cpp
--- a/Engine/Source/Runtime/TextureUtilitiesCommon/Private/TextureImportUtils.cpp
+++ b/Engine/Source/Runtime/TextureUtilitiesCommon/Private/TextureImportUtils.cpp
@@ -541,6 +541,15 @@ namespace UE
 
 		void FillZeroAlphaPNGData(int32 SizeX, int32 SizeY, ETextureSourceFormat SourceFormat, uint8* SourceData, bool bDoOnComplexAlphaNotJustBinaryTransparency)
 		{
+			// These conditions should be checked by IsImportResolutionValid, but just in case we get here
+			// via another path.
+			check(SizeX > 0 && SizeY > 0);
+			if (SizeX < 0 ||
+				SizeY < 0)
+			{
+				return;
+			}
+
 			switch (SourceFormat)
 			{
 				case TSF_BGRA8:
