TextureAssetActions followup cleanups and fixes + enable for cubemaps
Fix some Texture.h functions that should be const

#rnx

diff --git a/Engine/Source/Editor/UnrealEd/Private/TextureAssetActions.cpp b/Engine/Source/Editor/UnrealEd/Private/TextureAssetActions.cpp
--- a/Engine/Source/Editor/UnrealEd/Private/TextureAssetActions.cpp
+++ b/Engine/Source/Editor/UnrealEd/Private/TextureAssetActions.cpp
@@ -471,7 +471,8 @@ static void DoResizeTextureSource(UTexture * Texture,int TargetSize)
 	{
 		// not a warning, just FYI
 		// changing built size is totally possible and expected to happen sometimes
-		UE_LOG(LogTexture,Display,TEXT("DoResizeTextureSource failed to preserve built size; was: %dx%d now: %dx%d on [%s]"),
+		//	basically any time you resize smaller than the previous in-game size
+		UE_LOG(LogTexture,Verbose,TEXT("DoResizeTextureSource failed to preserve built size; was: %dx%d now: %dx%d on [%s]"),
 			BeforeSizeX,BeforeSizeY,
 			AfterSizeX,AfterSizeY,
 			*Texture->GetFullName());
@@ -538,7 +539,7 @@ void STextureAssetList::UpdateList()
 		Entry.Enabled = true;
 
 		ETextureClass Class = Texture->GetTextureClass();
-		if ( Class != ETextureClass::TwoD ) // @@ also ::Cube , Array ?
+		if ( Class != ETextureClass::TwoD && Class != ETextureClass::Cube ) //  Array ?
 		{
 			Status->WrongType = true;
 			Entry.Enabled = false;
diff --git a/Engine/Source/Runtime/Engine/Classes/Engine/Texture.h b/Engine/Source/Runtime/Engine/Classes/Engine/Texture.h
--- a/Engine/Source/Runtime/Engine/Classes/Engine/Texture.h
+++ b/Engine/Source/Runtime/Engine/Classes/Engine/Texture.h
@@ -1935,11 +1935,11 @@ public:
 	 * This downsizes using the mip generation system and so will only cut sizes in half. Return false
 	 * if a parameter is invalid. Returns true if the output is <= TargetSize, whether or not downsizing occurred.
 	 */
-	ENGINE_API bool DownsizeImageUsingTextureSettings(const ITargetPlatform* TargetPlatform, FImage& InOutImage, int32 TargetSize, int32 LayerIndex);
+	ENGINE_API bool DownsizeImageUsingTextureSettings(const ITargetPlatform* TargetPlatform, FImage& InOutImage, int32 TargetSize, int32 LayerIndex, bool & OutMadeChanges) const;
 
 	/*
 	*/
-	ENGINE_API void GetTargetPlatformBuildSettings(const ITargetPlatform* TargetPlatform, TArray<TArray<FTextureBuildSettings>>& OutSettingsPerFormatPerLayer);
+	ENGINE_API void GetTargetPlatformBuildSettings(const ITargetPlatform* TargetPlatform, TArray<TArray<FTextureBuildSettings>>& OutSettingsPerFormatPerLayer) const;
 
 #endif
 
diff --git a/Engine/Source/Runtime/Engine/Private/TextureDerivedData.cpp b/Engine/Source/Runtime/Engine/Private/TextureDerivedData.cpp
--- a/Engine/Source/Runtime/Engine/Private/TextureDerivedData.cpp
+++ b/Engine/Source/Runtime/Engine/Private/TextureDerivedData.cpp
@@ -3933,9 +3933,13 @@ void UTexture::SetMinTextureResidentMipCount(int32 InMinTextureResidentMipCount)
 }
 
 #if WITH_EDITOR
-bool UTexture::DownsizeImageUsingTextureSettings(const ITargetPlatform* TargetPlatform, FImage& InOutImage, int32 TargetSize, int32 LayerIndex)
+// return value false for critical errors
+// may return true even if nothing was done; check OutMadeChanges
+// InOutImage is modified in place ; output image will be same format but changed dimensions
+bool UTexture::DownsizeImageUsingTextureSettings(const ITargetPlatform* TargetPlatform, FImage& InOutImage, int32 TargetSize, int32 LayerIndex, bool & OutMadeChanges) const
 {
 	// resize so that the largest dimension is <= TargetSize
+	OutMadeChanges = false;
 
 	if (TargetSize <= 1 || LayerIndex < 0 || InOutImage.IsImageInfoValid() == false)
 	{
@@ -3963,11 +3967,31 @@ bool UTexture::DownsizeImageUsingTextureSettings(const ITargetPlatform* TargetPl
 
 	// Teak the build setting to generate a mip for our image
 	FTextureBuildSettings& BuildSettings = SettingPerLayer[LayerIndex];
+	// even if we are a Cube or LatLong, tell it we are just 2d ?
+	//  so the image is shrunk as a plain 2d
+	//	@@ not sure this is okay/best for latlongs
 	BuildSettings.bCubemap = false;
 	BuildSettings.bTextureArray = false;
 	BuildSettings.bVolume = false;
 	BuildSettings.bLongLatSource = false;
 
+	if ( BuildSettings.MipGenSettings == TMGS_NoMipmaps ||
+		BuildSettings.MipGenSettings == TMGS_LeaveExistingMips )
+	{
+		// what kind of mipgen do we use here?
+		// yuck!
+		BuildSettings.MipGenSettings = TMGS_SimpleAverage;
+	}
+
+	/*
+	// ?? yes ??
+	// make sure modern options are set:
+	BuildSettings.bUseNewMipFilter = true;
+	BuildSettings.bSharpenWithoutColorShift = false;
+	if ( IsNormalMap() )
+		BuildSettings.bNormalizeNormals = true;
+	*/
+
 	FImage Temp;
 	// convert to RGBA32F linear for the compressor
 	InOutImage.CopyTo(Temp, ERawImageFormat::RGBA32F, EGammaSpace::Linear);
@@ -3984,6 +4008,9 @@ bool UTexture::DownsizeImageUsingTextureSettings(const ITargetPlatform* TargetPl
 	constexpr int BuildSourceImageMipsMaxCount = 20; // plenty
 	BuildSourceImageMips.Empty(BuildSourceImageMipsMaxCount);
 
+	// one nice thing we do get from GenerateMipChain (as opposed to ResizeImage)
+	//	is that wrap/clamp address mode is respected and cubemaps clamp
+
 	ITextureCompressorModule::GenerateMipChain(BuildSettings, Temp, BuildSourceImageMips, 1);
 
 	// keep halving while larger size is > Target
@@ -4006,11 +4033,13 @@ bool UTexture::DownsizeImageUsingTextureSettings(const ITargetPlatform* TargetPl
 		SelectedOutput->CopyTo(InOutImage, InOutImage.Format, InOutImage.GammaSpace);
 	}
 
+	OutMadeChanges = true;
+
 	return true;
 }
 
 
-void UTexture::GetTargetPlatformBuildSettings(const ITargetPlatform* TargetPlatform, TArray<TArray<FTextureBuildSettings>>& OutSettingPerSupportedFormatPerLayer)
+void UTexture::GetTargetPlatformBuildSettings(const ITargetPlatform* TargetPlatform, TArray<TArray<FTextureBuildSettings>>& OutSettingPerSupportedFormatPerLayer) const
 {
 	ETextureEncodeSpeed EncodeSpeed = ETextureEncodeSpeed::Final;
 
diff --git a/Engine/Source/Runtime/TextureUtilitiesCommon/Private/TextureSourceDataUtils.cpp b/Engine/Source/Runtime/TextureUtilitiesCommon/Private/TextureSourceDataUtils.cpp
--- a/Engine/Source/Runtime/TextureUtilitiesCommon/Private/TextureSourceDataUtils.cpp
+++ b/Engine/Source/Runtime/TextureUtilitiesCommon/Private/TextureSourceDataUtils.cpp
@@ -18,12 +18,7 @@ namespace Private
 	// resize so that the largest dimension is <= MaxSize
 	bool ResizeTexture2D(UTexture* Texture, int32 MaxSize, const ITargetPlatform* TargetPlatform)
 	{
-		// Protect the code from an async build of the texture
-		// @@ I think that this PreEditChange can (and should) be moved to after DownsizeImageUsingTextureSettings
-		Texture->PreEditChange(nullptr);
-
 		// We want to reduce the asset size so ignore the imported mip(s)
-		//	@@ should not do this if MipGen == LeaveExisting
 		const int32 MipIndex = 0;
 		FImage SourceMip0;
 		if (!Texture->Source.GetMipImage(SourceMip0, MipIndex))
@@ -31,14 +26,25 @@ namespace Private
 			return false;
 		}
 
+		int32 NumSlices = Texture->Source.GetNumSlices(); // == 1 or 6 for cubes
+
 		const int32 LayerIndex = 0;
-		// DownsizeImageUsingTextureSettings is not const but I think it can and should be?
-		//	(the Texture itself is not changed, only the FImage argument is changed)
-		if ( ! Texture->DownsizeImageUsingTextureSettings(TargetPlatform, SourceMip0, MaxSize, LayerIndex) )
+		bool MadeChanges;
+		if ( ! Texture->DownsizeImageUsingTextureSettings(TargetPlatform, SourceMip0, MaxSize, LayerIndex, MadeChanges) )
 		{
+			// a critical error
 			return false;
 		}
-		FImage ResizedImage = MoveTemp(SourceMip0);
+		if ( ! MadeChanges )
+		{
+			return false;
+		}
+		
+		Texture->PreEditChange(nullptr);
+
+		FImage ResizedImage = MoveTemp(SourceMip0); // this is just a variable rename
+
+		check( ResizedImage.NumSlices == NumSlices ); // slices are done one by one
 
 		UE::Serialization::FEditorBulkData::FSharedBufferWithID ResizedImageBufferWithID = MakeSharedBufferFromArray(MoveTemp(ResizedImage.RawData));
 
@@ -58,26 +64,28 @@ namespace Private
 		return true;
 	}
 
-	bool ResizeTextureSlicedBy2DLayers(UTexture* Texture, int32 MaxSize, const ITargetPlatform* TargetPlatform)
+	bool ResizeTextureSlicesOneByOne(UTexture* Texture, int32 MaxSize, const ITargetPlatform* TargetPlatform)
 	{
+		// @@ delete me; this function is not used; the normal ResizeTexture2D does the same thing
+
 		// should check Source size vs MaxSize and early return here
 		if ( Texture->Source.GetSizeX() <= MaxSize && Texture->Source.GetSizeY() <= MaxSize )
 			return false;
 				
-		// Protect the code from an async build of the texture
-		Texture->PreEditChange(nullptr);
-
 		FImage ResizedImage;
+
 		{
+			check( Texture->Source.GetNumSlices() > 1 );
+
 			TArray<FImage> ResizedSlices;
 			ResizedSlices.Reserve(Texture->Source.GetNumSlices());
 
 			ERawImageFormat::Type FormatUsed;
 			EGammaSpace GammaSpaceUsed;
+			bool MadeAnyChanges = false;
 
 			{
 				// We want to reduce the asset size so ignore the imported mip(s)
-				//	@@ should not do this if MipGen == LeaveExisting
 				TArray<FImage> Slices;
 				Slices.Reserve(Texture->Source.GetNumSlices());
 
@@ -92,28 +100,42 @@ namespace Private
 					FormatUsed = SourceMip0.Format;
 					GammaSpaceUsed = SourceMip0.GammaSpace;
 
+					check( SourceMip0.NumSlices == Texture->Source.GetNumSlices() );
+
 					for (int32 Index = 0; Index < SourceMip0.NumSlices; ++Index)
 					{
 						FImage& Slice = Slices.AddDefaulted_GetRef();
 						FImageView SliceView = SourceMip0.GetSlice(Index);
 
-						FImageCore::CopyImage(SliceView, Slice);
+						SliceView.CopyTo(Slice); // allocs new image in Slice
 					}
 				}
 
+				check( Slices.Num() == Texture->Source.GetNumSlices() );
+
 				for (FImage& Slice : Slices)
 				{
-					FImage& ResizedSlice = ResizedSlices.AddDefaulted_GetRef();
-
 					const int32 LayerIndex = 0;
-					if ( ! Texture->DownsizeImageUsingTextureSettings(TargetPlatform, Slice, MaxSize, LayerIndex) )
+					bool MadeChanges;
+					if ( ! Texture->DownsizeImageUsingTextureSettings(TargetPlatform, Slice, MaxSize, LayerIndex, MadeChanges) )
 					{
+						// a critical error
 						return false;
 					}
+					MadeAnyChanges = MadeAnyChanges || MadeChanges;
+					
+					FImage& ResizedSlice = ResizedSlices.AddDefaulted_GetRef();
 					ResizedSlice = MoveTemp(Slice);
 				}
 			}
 
+			check( ResizedSlices.Num() == Texture->Source.GetNumSlices() );
+			
+			if ( ! MadeAnyChanges )
+			{
+				return false;
+			}
+
 			// Move the resized slices into the resized image
 			ResizedImage.Init(ResizedSlices[0].SizeX,ResizedSlices[0].SizeY, ResizedSlices.Num(), FormatUsed, GammaSpaceUsed);
 
@@ -122,7 +144,9 @@ namespace Private
 				FImageCore::CopyImage(ResizedSlices[Index], ResizedImage.GetSlice(Index));
 			}
 		}
-
+		
+		// Protect the code from an async build of the texture
+		Texture->PreEditChange(nullptr);
 
 		UE::Serialization::FEditorBulkData::FSharedBufferWithID ResizedImageBufferWithID = MakeSharedBufferFromArray(MoveTemp(ResizedImage.RawData));
 
@@ -133,14 +157,17 @@ namespace Private
 			, NumMips
 			, FImageCoreUtils::ConvertToTextureSourceFormat(ResizedImage.Format)
 			, MoveTemp(ResizedImageBufferWithID));
+			
+		// if gamma was Pow22 it is now sRGB
+		Texture->bUseLegacyGamma = false;
 
 		return true;
 	}
 
 	bool ResizeTexture2DBlocked(UTexture* Texture, int32 MaxSize, const ITargetPlatform* TargetPlatform)
 	{
-		// Protect the code from an async build of the texture
-		Texture->PreEditChange(nullptr);
+		// note: does not support layers
+		// MaxSize is applied to each block in the UDIM, not the total size
 
 		/*
 		FIntPoint LogicalSourceSize = Texture->Source.GetLogicalSize();
@@ -160,6 +187,8 @@ namespace Private
 		TArray<FImage> ResizedBlocks;
 		ResizedBlocks.Reserve(Texture->Source.GetNumBlocks());
 
+		bool MadeAnyChanges = false;
+
 		for (int32 BlockIndex = 0; BlockIndex < Texture->Source.GetNumBlocks(); ++BlockIndex)
 		{
 			// We want to reduce the asset size so ignore the imported mip(s)
@@ -171,18 +200,21 @@ namespace Private
 				return false;
 			}
 
-			FTextureSourceBlock& ResizedSourceBlock = ResizedSourceBlocks.AddDefaulted_GetRef();
-			Texture->Source.GetBlock(BlockIndex, ResizedSourceBlock);
-		
-			FImage& ResizedBlock = ResizedBlocks.AddDefaulted_GetRef();
-
 			//int32 BlockMaxSize = FMath::RoundToInt32(FMath::Min(ResizedSourceBlock.SizeX * RatioX, ResizedSourceBlock.SizeY * RatioY));
 
 			// each block is resized to MaxSize
-			if ( ! Texture->DownsizeImageUsingTextureSettings(TargetPlatform, SourceMip0, MaxSize, LayerIndex) )
+			bool MadeChanges;
+			if ( ! Texture->DownsizeImageUsingTextureSettings(TargetPlatform, SourceMip0, MaxSize, LayerIndex, MadeChanges) )
 			{
+				// critical error
 				return false;
 			}
+			MadeAnyChanges = MadeAnyChanges || MadeChanges;
+
+			FTextureSourceBlock& ResizedSourceBlock = ResizedSourceBlocks.AddDefaulted_GetRef();
+			Texture->Source.GetBlock(BlockIndex, ResizedSourceBlock);
+		
+			FImage& ResizedBlock = ResizedBlocks.AddDefaulted_GetRef();
 			ResizedBlock = MoveTemp(SourceMip0);
 			
 			ResizedSourceBlock.SizeX = ResizedBlock.SizeX;
@@ -190,10 +222,18 @@ namespace Private
 			ResizedSourceBlock.NumSlices = 1;
 		}
 
+		if ( ! MadeAnyChanges )
+		{
+			return false;
+		}
+		
+		// Protect the code from an async build of the texture
+		Texture->PreEditChange(nullptr);
+
 		int64 SizeNeededInBytes = 0;
 		for (const FImage& Block : ResizedBlocks)
 		{
-			SizeNeededInBytes += Block.RawData.Num();
+			SizeNeededInBytes += Block.RawData.Num(); // Block.GetImageSizeBytes()
 		}
 		FUniqueBuffer WriteImageBuffer = FUniqueBuffer::Alloc(SizeNeededInBytes);
 
@@ -216,6 +256,9 @@ namespace Private
 			MoveTemp(ResizedImageBufferWithID)
 		);
 
+		// if gamma was Pow22 it is now sRGB
+		Texture->bUseLegacyGamma = false;
+
 		return true;
 	}
 }
@@ -226,6 +269,19 @@ bool DownsizeTextureSourceData(UTexture* Texture, int32 TargetSizeInGame, const
 	check( Texture->Source.IsValid() );
 
 	// Check if we don't know how to resize that texture
+
+	if ( Texture->Source.GetNumMips() > 1 && Texture->MipGenSettings == TMGS_LeaveExistingMips )
+	{
+		//	should not do this if MipGen == LeaveExisting ; or maybe warn?
+		// return false;
+
+		// go ahead and do it, but warn:
+
+		UE_LOG(LogTexture,Warning,TEXT("DownsizeTextureSourceData: Texture has LeaveExistingMips ; they will be discarded! [%s]"),
+			*Texture->GetFullName());
+	}
+
+	// we only support 1 layer currently
 	if (Texture->Source.GetNumLayers() != 1)
 	{
 		return false;
@@ -290,7 +346,8 @@ bool DownsizeTextureSourceData(UTexture* Texture, int32 TargetSizeInGame, const
 		}
 		else
 		{
-			return Private::ResizeTextureSlicedBy2DLayers(Texture, TargetSourceSize, TargetPlatform);
+			//return Private::ResizeTextureSlicesOneByOne(Texture, TargetSourceSize, TargetPlatform);
+			return Private::ResizeTexture2D(Texture, TargetSourceSize, TargetPlatform);
 		}
 	}
 	// could do GetTextureClass == Array ?
