TextureSourceDataUtils : remove unused resizer OneByOne and fix cube map size inconsistency

#rnx

diff --git a/Engine/Source/Runtime/TextureUtilitiesCommon/Private/TextureSourceDataUtils.cpp b/Engine/Source/Runtime/TextureUtilitiesCommon/Private/TextureSourceDataUtils.cpp
--- a/Engine/Source/Runtime/TextureUtilitiesCommon/Private/TextureSourceDataUtils.cpp
+++ b/Engine/Source/Runtime/TextureUtilitiesCommon/Private/TextureSourceDataUtils.cpp
@@ -64,110 +64,12 @@ namespace Private
 		return true;
 	}
 
-	bool ResizeTextureSlicesOneByOne(UTexture* Texture, int32 MaxSize, const ITargetPlatform* TargetPlatform)
-	{
-		// @@ delete me; this function is not used; the normal ResizeTexture2D does the same thing
-
-		// should check Source size vs MaxSize and early return here
-		if ( Texture->Source.GetSizeX() <= MaxSize && Texture->Source.GetSizeY() <= MaxSize )
-			return false;
-				
-		FImage ResizedImage;
-
-		{
-			check( Texture->Source.GetNumSlices() > 1 );
-
-			TArray<FImage> ResizedSlices;
-			ResizedSlices.Reserve(Texture->Source.GetNumSlices());
-
-			ERawImageFormat::Type FormatUsed;
-			EGammaSpace GammaSpaceUsed;
-			bool MadeAnyChanges = false;
-
-			{
-				// We want to reduce the asset size so ignore the imported mip(s)
-				TArray<FImage> Slices;
-				Slices.Reserve(Texture->Source.GetNumSlices());
-
-				{
-					// We could probably avoid a copy here but for the simplicity of the code keep it for now.
-					FImage SourceMip0;
-					if (!Texture->Source.GetMipImage(SourceMip0, 0))
-					{
-						return false;
-					}
-
-					FormatUsed = SourceMip0.Format;
-					GammaSpaceUsed = SourceMip0.GammaSpace;
-
-					check( SourceMip0.NumSlices == Texture->Source.GetNumSlices() );
-
-					for (int32 Index = 0; Index < SourceMip0.NumSlices; ++Index)
-					{
-						FImage& Slice = Slices.AddDefaulted_GetRef();
-						FImageView SliceView = SourceMip0.GetSlice(Index);
-
-						SliceView.CopyTo(Slice); // allocs new image in Slice
-					}
-				}
-
-				check( Slices.Num() == Texture->Source.GetNumSlices() );
-
-				for (FImage& Slice : Slices)
-				{
-					const int32 LayerIndex = 0;
-					bool MadeChanges;
-					if ( ! Texture->DownsizeImageUsingTextureSettings(TargetPlatform, Slice, MaxSize, LayerIndex, MadeChanges) )
-					{
-						// a critical error
-						return false;
-					}
-					MadeAnyChanges = MadeAnyChanges || MadeChanges;
-					
-					FImage& ResizedSlice = ResizedSlices.AddDefaulted_GetRef();
-					ResizedSlice = MoveTemp(Slice);
-				}
-			}
-
-			check( ResizedSlices.Num() == Texture->Source.GetNumSlices() );
-			
-			if ( ! MadeAnyChanges )
-			{
-				return false;
-			}
-
-			// Move the resized slices into the resized image
-			ResizedImage.Init(ResizedSlices[0].SizeX,ResizedSlices[0].SizeY, ResizedSlices.Num(), FormatUsed, GammaSpaceUsed);
-
-			for (int32 Index = 0; Index < ResizedImage.NumSlices; ++Index)
-			{
-				FImageCore::CopyImage(ResizedSlices[Index], ResizedImage.GetSlice(Index));
-			}
-		}
-		
-		// Protect the code from an async build of the texture
-		Texture->PreEditChange(nullptr);
-
-		UE::Serialization::FEditorBulkData::FSharedBufferWithID ResizedImageBufferWithID = MakeSharedBufferFromArray(MoveTemp(ResizedImage.RawData));
-
-		const int32 NumMips = 1;
-		Texture->Source.Init(ResizedImage.SizeX
-			, ResizedImage.SizeY
-			, ResizedImage.NumSlices
-			, NumMips
-			, FImageCoreUtils::ConvertToTextureSourceFormat(ResizedImage.Format)
-			, MoveTemp(ResizedImageBufferWithID));
-			
-		// if gamma was Pow22 it is now sRGB
-		Texture->bUseLegacyGamma = false;
-
-		return true;
-	}
-
 	bool ResizeTexture2DBlocked(UTexture* Texture, int32 MaxSize, const ITargetPlatform* TargetPlatform)
 	{
 		// note: does not support layers
+
 		// MaxSize is applied to each block in the UDIM, not the total size
+		// @@ should we target total size?
 
 		/*
 		FIntPoint LogicalSourceSize = Texture->Source.GetLogicalSize();
@@ -264,7 +166,7 @@ namespace Private
 }
 
 
-bool DownsizeTextureSourceData(UTexture* Texture, int32 TargetSizeInGame, const ITargetPlatform* TargetPlatform)
+bool DownsizeTextureSourceData(UTexture* Texture, int32 TargetSourceSize, const ITargetPlatform* TargetPlatform)
 {
 	check( Texture->Source.IsValid() );
 
@@ -286,14 +188,10 @@ bool DownsizeTextureSourceData(UTexture* Texture, int32 TargetSizeInGame, const
 	{
 		return false;
 	}
-
+	
 	if (!(Texture->Source.GetTextureClass() == ETextureClass::Cube || Texture->Source.GetTextureClass() == ETextureClass::TwoD))
 	{
-		return false;
-	}
-
-	if (Texture->Source.GetTextureClass() == ETextureClass::TwoD && Texture->Source.GetNumSlices() != 1)
-	{
+		// array, cubearray, volume, not supported
 		return false;
 	}
 
@@ -312,47 +210,22 @@ bool DownsizeTextureSourceData(UTexture* Texture, int32 TargetSizeInGame, const
 		}
 	}
 
-	int32 TargetSourceSize = TargetSizeInGame;
-	if (Texture->Source.IsLongLatCubemap())
-	{
-		// The function return the max size of the generated cube from the source long lat
-		// this should be kept in sync with the implementation details of ComputeLongLatCubemapExtents() or refactored
-		TargetSourceSize = (1U << FMath::FloorLog2(TargetSizeInGame)) * 2;
-	}
-
 	FIntPoint SourceSize = Texture->Source.GetLogicalSize();
 	if (SourceSize.X <= TargetSourceSize && SourceSize.Y <= TargetSourceSize)
 	{
 		return false;
 	}
 
-	if (Texture->Source.GetTextureClass() == ETextureClass::TwoD)
+	if (Texture->Source.GetNumBlocks() == 1)
 	{
-		if (Texture->Source.GetNumBlocks() == 1)
-		{
-			return Private::ResizeTexture2D(Texture, TargetSourceSize, TargetPlatform);
-		}
-		else
-		{
-			// UDIM(s)
-			return Private::ResizeTexture2DBlocked(Texture, TargetSourceSize, TargetPlatform);
-		}
+		return Private::ResizeTexture2D(Texture, TargetSourceSize, TargetPlatform);
 	}
-	else if (Texture->Source.GetTextureClass() == ETextureClass::Cube)
+	else
 	{
-		if (Texture->Source.IsLongLatCubemap())
-		{
-			return Private::ResizeTexture2D(Texture, TargetSourceSize, TargetPlatform);
-		}
-		else
-		{
-			//return Private::ResizeTextureSlicesOneByOne(Texture, TargetSourceSize, TargetPlatform);
-			return Private::ResizeTexture2D(Texture, TargetSourceSize, TargetPlatform);
-		}
+		// UDIM VT
+		return Private::ResizeTexture2DBlocked(Texture, TargetSourceSize, TargetPlatform);
 	}
-	// could do GetTextureClass == Array ?
-	// other classes unsupported
-	
+		
 	return false;
 }
 
@@ -367,8 +240,17 @@ bool DownsizeTexureSourceDataNearRenderingSize(UTexture* Texture, const ITargetP
 	int32 BeforeSizeY;
 	Texture->GetBuiltTextureSize(TargetPlatform, BeforeSizeX, BeforeSizeY);
 
-	int32 TargetSize = FMath::Max(BeforeSizeX, BeforeSizeY);
-	if (DownsizeTextureSourceData(Texture, TargetSize, TargetPlatform))
+	int32 TargetSizeInGame = FMath::Max(BeforeSizeX, BeforeSizeY);
+	
+	int32 TargetSourceSize = TargetSizeInGame;
+	if (Texture->Source.IsLongLatCubemap())
+	{
+		// The function return the max size of the generated cube from the source long lat
+		// this should be kept in sync with the implementation details of ComputeLongLatCubemapExtents() or refactored
+		TargetSourceSize = (1U << FMath::FloorLog2(TargetSizeInGame)) * 2;
+	}
+
+	if (DownsizeTextureSourceData(Texture, TargetSourceSize, TargetPlatform))
 	{
 		Texture->LODBias = 0;
 		Texture->PostEditChange();
