TextureAssetActions : fix VT UDIM sizing inconsistency; resize to fit total udim size in threshold
TextureEditor : show udim info in imported size
also treat these actions as a reimport (bump oodle sdk to latest)

#rnx

diff --git a/Engine/Source/Editor/TextureEditor/Private/TextureEditorToolkit.cpp b/Engine/Source/Editor/TextureEditor/Private/TextureEditorToolkit.cpp
--- a/Engine/Source/Editor/TextureEditor/Private/TextureEditorToolkit.cpp
+++ b/Engine/Source/Editor/TextureEditor/Private/TextureEditorToolkit.cpp
@@ -762,10 +762,15 @@ void FTextureEditorToolkit::PopulateQuickInfo( )
 	const int32 NumSurfaces = Texture->GetSurfaceArraySize();
 	const int32 ArraySize = bIsArray ? (bIsCube ? NumSurfaces / 6 : NumSurfaces) : 0;
 
+	// this is the size of just one block :
+	//	use GetLogicalSize for all blocks
 	const int32 ImportedWidth = Texture->Source.IsValid() ? Texture->Source.GetSizeX() : SurfaceWidth;
 	const int32 ImportedHeight = Texture->Source.IsValid() ? Texture->Source.GetSizeY() : SurfaceHeight;
 	const int32 ImportedDepth = Texture->Source.IsValid() ? (bIsVolume ? Texture->Source.GetNumSlices() : 0) : SurfaceDepth;
 
+	const int32 ImportedNumBlocks = Texture->Source.GetNumBlocks();
+	const FIntPoint ImportedSizeInBlocks = Texture->Source.GetSizeInBlocks();
+
 	const FStreamableRenderResourceState SRRState = Texture->GetStreamableResourceState();
 	const int32 ActualMipBias = SRRState.IsValid() ? (SRRState.ResidentFirstLODIdx() + SRRState.AssetLODBias) : Texture->GetCachedLODBias();
 	// Editor dimensions (takes user specified mip setting into account)
@@ -836,7 +841,17 @@ void FTextureEditorToolkit::PopulateQuickInfo( )
 	}
 	else
 	{
-	    ImportedText->SetText(FText::Format( NSLOCTEXT("TextureEditor", "QuickInfo_Imported_2x", "Imported: {0}x{1}{2}"), FText::AsNumber(ImportedWidth, &FormatOptions), FText::AsNumber(ImportedHeight, &FormatOptions), ImportedCubemapInfo));
+		if ( ImportedNumBlocks > 1 )
+		{
+		    ImportedText->SetText(FText::Format( NSLOCTEXT("TextureEditor", "QuickInfo_Imported_Blocks", "Imported: {0}x{1}x{2}x{3}"), 
+				FText::AsNumber(ImportedWidth, &FormatOptions), FText::AsNumber(ImportedHeight, &FormatOptions), 
+				FText::AsNumber(ImportedSizeInBlocks.X, &FormatOptions), FText::AsNumber(ImportedSizeInBlocks.Y, &FormatOptions)));
+		}
+		else
+		{
+		    ImportedText->SetText(FText::Format( NSLOCTEXT("TextureEditor", "QuickInfo_Imported_2x", "Imported: {0}x{1}{2}"), FText::AsNumber(ImportedWidth, &FormatOptions), FText::AsNumber(ImportedHeight, &FormatOptions), ImportedCubemapInfo));
+		}
+
 		CurrentText->SetText(FText::Format( NSLOCTEXT("TextureEditor", "QuickInfo_Displayed_2x", "Displayed: {0}x{1}{2}"), FText::AsNumber(PreviewEffectiveTextureWidth, &FormatOptions ), FText::AsNumber(PreviewEffectiveTextureHeight, &FormatOptions), DisplayedCubemapInfo));
 		MaxInGameText->SetText(FText::Format( NSLOCTEXT("TextureEditor", "QuickInfo_MaxInGame_2x", "Max In-Game: {0}x{1}{2}"), FText::AsNumber(MaxInGameWidth, &FormatOptions), FText::AsNumber(MaxInGameHeight, &FormatOptions), InGameCubemapInfo));
 	}
diff --git a/Engine/Source/Editor/UnrealEd/Private/TextureAssetActions.cpp b/Engine/Source/Editor/UnrealEd/Private/TextureAssetActions.cpp
--- a/Engine/Source/Editor/UnrealEd/Private/TextureAssetActions.cpp
+++ b/Engine/Source/Editor/UnrealEd/Private/TextureAssetActions.cpp
@@ -30,6 +30,7 @@
 #include "Engine/Texture.h"
 #include "Interfaces/ITargetPlatformManagerModule.h"
 #include "TextureSourceDataUtils.h"
+#include "TextureImportSettings.h"
 
 #define LOCTEXT_NAMESPACE "TextureAssetActions"
 
@@ -168,7 +169,6 @@ private:
 	TArray<TextureListEntry> TextureList;
 
 	int ThresholdValue = 0;
-	bool bThresholdVisible = true;
 
 	bool bNormalMapsKeep16bits = true;
 
@@ -368,8 +368,6 @@ void STextureAssetList::Init(const TArray<UTexture *> &Textures, ETextureAction
 {
 	Action = InAction;
 	
-	bThresholdVisible = (Action == ETextureAction::Resize );
-
 	IntroMessage = TAA_Intro(Action);
 
 	TextureList.SetNum(Textures.Num());
@@ -577,6 +575,9 @@ static void DoResizeTextureSource(UTexture * Texture,int TargetSize)
 			*Texture->GetFullName());
 	}
 
+	// this counts as a reimport :
+	UE::TextureUtilitiesCommon::ApplyDefaultsForNewlyImportedTextures(Texture,true);
+
 	// DownsizeTextureSourceData did the PreEditChange
 	Texture->PostEditChange();
 }
@@ -601,6 +602,7 @@ static ETextureSourceFormat GetReducedTextureSourceFormat(TextureCompressionSett
 		// Gray and Displacement pass through G16 ; note they do not do that for RGBA16 (see GetDefaultTextureFormatName)
 		if ( InTSF == TSF_G16 ) return InTSF;
 		// otherwise we will convert to G8
+		// [[fallthrough]];
 	case TC_DistanceFieldFont		: //"DistanceFieldFont (G8)"),
 		Out8bit = true;
 		OutSingleChannel = true;
@@ -807,8 +809,7 @@ void STextureAssetList::UpdateList()
 			}
 			else
 			{
-				// note: this does not match the logic in TextureSourceDataUtils
-				//	(@@FIXME) (eg. udim)
+				// GetLogicalSize = sum of udim blocks
 				FIntPoint SourceSize = Texture->Source.GetLogicalSize();
 				int MaxSize = FMath::Max(SourceSize.X,SourceSize.Y);
 				if ( MaxSize <= ThresholdValue )
@@ -827,6 +828,7 @@ void STextureAssetList::UpdateList()
 			}
 			else if ( Texture->Source.GetSourceCompression() == ETextureSourceCompressionFormat::TSCF_JPEG )
 			{
+				// JPEG is already 8 bit; we don't want to try to change RGB JPEG to G8, just leave it alone
 				Status = EAssetActionStatus::DontChangeJPEG;
 			}
 		}
@@ -928,7 +930,7 @@ EVisibility STextureAssetList::GetErrorMessageVisibility() const
 
 EVisibility STextureAssetList::GetThresholdVisibility() const
 {
-	return (bThresholdVisible) ? EVisibility::Visible : EVisibility::Collapsed;
+	return ( Action == ETextureAction::Resize ) ? EVisibility::Visible : EVisibility::Collapsed;
 }
 
 FText STextureAssetList::GetIntroMessage() const
diff --git a/Engine/Source/Runtime/TextureUtilitiesCommon/Private/TextureSourceDataUtils.cpp b/Engine/Source/Runtime/TextureUtilitiesCommon/Private/TextureSourceDataUtils.cpp
--- a/Engine/Source/Runtime/TextureUtilitiesCommon/Private/TextureSourceDataUtils.cpp
+++ b/Engine/Source/Runtime/TextureUtilitiesCommon/Private/TextureSourceDataUtils.cpp
@@ -8,6 +8,7 @@
 #include "Engine/Texture.h"
 #include "HAL/UnrealMemory.h"
 #include "EngineLogs.h"
+#include "TextureImportSettings.h"
 
 namespace UE::TextureUtilitiesCommon::Experimental
 {
@@ -18,7 +19,10 @@ namespace Private
 	// resize so that the largest dimension is <= MaxSize
 	static bool ResizeTexture2D(UTexture* Texture, int32 MaxSize, const ITargetPlatform* TargetPlatform)
 	{
-		// We want to reduce the asset size so ignore the imported mip(s)
+		check( Texture->Source.GetNumLayers() == 1 );
+		const int32 LayerIndex = 0;
+
+		// We want to reduce the asset size so ignore the imported mip(s) (??)
 		const int32 MipIndex = 0;
 		FImage SourceMip0;
 		if (!Texture->Source.GetMipImage(SourceMip0, MipIndex))
@@ -30,7 +34,6 @@ namespace Private
 
 		int32 NumSlices = Texture->Source.GetNumSlices(); // == 1 or 6 for cubes
 
-		const int32 LayerIndex = 0;
 		bool MadeChanges;
 		if ( ! Texture->DownsizeImageUsingTextureSettings(TargetPlatform, SourceMip0, MaxSize, LayerIndex, MadeChanges) )
 		{
@@ -88,24 +91,20 @@ namespace Private
 		return WriteImageBuffer.MoveToShared();
 	}
 
-	static bool ResizeTexture2DBlocked(UTexture* Texture, int32 MaxSize, const ITargetPlatform* TargetPlatform)
+	static bool ResizeTexture2DBlocked(UTexture* Texture, int32 TotalMaxSize, const ITargetPlatform* TargetPlatform)
 	{
-		// note: does not support layers
+		// does not support layers
+		check( Texture->Source.GetNumLayers() == 1 );
+		const int32 NumLayers = 1;
+		const int32 LayerIndex = 0;
 
-		// MaxSize is applied to each block in the UDIM, not the total size
-		// @@ should we target total size?
+		// MaxSize is applied to the total UDIM size
 
-		/*
 		FIntPoint LogicalSourceSize = Texture->Source.GetLogicalSize();
-		double RatioX = double(MaxSize) / LogicalSourceSize.X;
-		double RatioY = double(MaxSize) / LogicalSourceSize.Y;
+		check( LogicalSourceSize.X > TotalMaxSize || LogicalSourceSize.Y > TotalMaxSize );
 
-		// early return if we're not shrinking
-		if ( RatioX >= 1.0 && RatioY >= 1.0 )
-		{
-			return false;
-		}
-		*/
+		double ResizeRatio = double(TotalMaxSize) / FMath::Max(LogicalSourceSize.X,LogicalSourceSize.Y);
+		check( ResizeRatio < 1.0 );
 
 		TArray<FTextureSourceBlock> ResizedSourceBlocks;
 		ResizedSourceBlocks.Reserve(Texture->Source.GetNumBlocks());
@@ -120,7 +119,6 @@ namespace Private
 			// We want to reduce the asset size so ignore the imported mip(s)
 			FImage SourceMip0;
 			const int32 MipIndex = 0;
-			const int32 LayerIndex = 0;
 			if (!Texture->Source.GetMipImage(SourceMip0, BlockIndex, LayerIndex, MipIndex))
 			{
 				UE_LOG(LogTexture,Error,TEXT("ResizeTexture2DBlocked: Texture GetMipImage failed [%s]"),
@@ -128,11 +126,12 @@ namespace Private
 				return false;
 			}
 
-			//int32 BlockMaxSize = FMath::RoundToInt32(FMath::Min(ResizedSourceBlock.SizeX * RatioX, ResizedSourceBlock.SizeY * RatioY));
+			int32 NewSizeX = FMath::RoundToInt32( ResizeRatio * SourceMip0.SizeX );
+			int32 NewSizeY = FMath::RoundToInt32( ResizeRatio * SourceMip0.SizeY );
+			int32 BlockMaxSize = FMath::Max(NewSizeX,NewSizeY);
 
-			// each block is resized to MaxSize
 			bool MadeChanges;
-			if ( ! Texture->DownsizeImageUsingTextureSettings(TargetPlatform, SourceMip0, MaxSize, LayerIndex, MadeChanges) )
+			if ( ! Texture->DownsizeImageUsingTextureSettings(TargetPlatform, SourceMip0, BlockMaxSize, LayerIndex, MadeChanges) )
 			{
 				// critical error
 				UE_LOG(LogTexture,Error,TEXT("ResizeTexture2DBlocked: Texture DownsizeImageUsingTextureSettings failed [%s]"),
@@ -149,7 +148,7 @@ namespace Private
 			
 			ResizedSourceBlock.SizeX = ResizedBlock.SizeX;
 			ResizedSourceBlock.SizeY = ResizedBlock.SizeY;
-			ResizedSourceBlock.NumSlices = 1;
+			ResizedSourceBlock.NumMips = 1;
 		}
 
 		if ( ! MadeAnyChanges )
@@ -163,7 +162,6 @@ namespace Private
 		UE::Serialization::FEditorBulkData::FSharedBufferWithID ResizedImageBufferWithID = MakeSharedBufferForImageDatas(ResizedBlocks);
 
 		const ETextureSourceFormat SourceFormat = Texture->Source.GetFormat();
-		int32 NumLayers = 1;
 		Texture->Source.InitBlocked(
 			&SourceFormat, // array of formats per layer
 			ResizedSourceBlocks.GetData(),
@@ -267,6 +265,10 @@ TEXTUREUTILITIESCOMMON_API bool DownsizeTextureSourceDataNearRenderingSize(UText
 	if (DownsizeTextureSourceData(Texture, TargetSourceSize, TargetPlatform))
 	{
 		Texture->LODBias = 0;
+		
+		// this counts as a reimport :
+		UE::TextureUtilitiesCommon::ApplyDefaultsForNewlyImportedTextures(Texture,true);
+
 		Texture->PostEditChange();
 		
 		// check that GetBuiltTextureSize was preserved :
@@ -339,7 +341,7 @@ TEXTUREUTILITIESCOMMON_API bool ChangeTextureSourceFormat(UTexture* Texture, ETe
 			, NewFormat
 			, MoveTemp(ResizedImageBufferWithID));
 	}
-	else
+	else // blocks and/or mips
 	{
 		// all blocks of a UDIM have the same format; Layers do not
 		int32 NumLayers = 1;
@@ -394,6 +396,10 @@ TEXTUREUTILITIESCOMMON_API bool ChangeTextureSourceFormat(UTexture* Texture, ETe
 
 	// if gamma was Pow22 it is now sRGB
 	Texture->bUseLegacyGamma = false;
+	
+	// this counts as a reimport :
+	UE::TextureUtilitiesCommon::ApplyDefaultsForNewlyImportedTextures(Texture,true);
+
 	Texture->PostEditChange();
 
 	check( Texture->Source.GetGammaSpace(0) == NewGamma );
