Texture Source resizing followups
ResizeImage support wrap/clamp edge mode on X and Y
ResizeImage fix Pow22 source is SRGB dest
TextureSource resizer detect cases that are poorly handled by TextureCompressorModule and use ResizeImage instead (eg. latlong, stretch)

diff --git a/Engine/Source/Runtime/Engine/Private/TextureDerivedData.cpp b/Engine/Source/Runtime/Engine/Private/TextureDerivedData.cpp
--- a/Engine/Source/Runtime/Engine/Private/TextureDerivedData.cpp
+++ b/Engine/Source/Runtime/Engine/Private/TextureDerivedData.cpp
@@ -3969,28 +3969,39 @@ bool UTexture::DownsizeImageUsingTextureSettings(const ITargetPlatform* TargetPl
 	FTextureBuildSettings& BuildSettings = SettingPerLayer[LayerIndex];
 	// even if we are a Cube or LatLong, tell it we are just 2d ?
 	//  so the image is shrunk as a plain 2d
-	//	@@ not sure this is okay/best for latlongs
 	BuildSettings.bCubemap = false;
 	BuildSettings.bTextureArray = false;
 	BuildSettings.bVolume = false;
 	BuildSettings.bLongLatSource = false;
+	
+	// make sure modern options are set:
+	BuildSettings.bUseNewMipFilter = true;
+	BuildSettings.bSharpenWithoutColorShift = false;
+	if ( IsNormalMap() )
+	{
+		BuildSettings.bNormalizeNormals = true;
+	}
 
 	if ( BuildSettings.MipGenSettings == TMGS_NoMipmaps ||
-		BuildSettings.MipGenSettings == TMGS_LeaveExistingMips )
+		BuildSettings.MipGenSettings == TMGS_LeaveExistingMips ||
+		BuildSettings.MipGenSettings == TMGS_Angular )
 	{
-		// what kind of mipgen do we use here?
-		// yuck!
+		// what kind of mipgen do we use here? (default from GetMipGenSettings will use 2x2 simple average)
+		// external caller now prefers to use use ResizeImage in this case
 		BuildSettings.MipGenSettings = TMGS_SimpleAverage;
 	}
 
-	/*
-	// ?? yes ??
-	// make sure modern options are set:
-	BuildSettings.bUseNewMipFilter = true;
-	BuildSettings.bSharpenWithoutColorShift = false;
-	if ( IsNormalMap() )
-		BuildSettings.bNormalizeNormals = true;
-	*/
+	// we turned off bCubeMap, make sure cube face filters clamp, not wrap
+	//  see ComputeAddressMode
+	if ( GetTextureClass() == ETextureClass::Cube || GetTextureClass() == ETextureClass::CubeArray )
+	{
+		// for 6-face cubes, just clamp
+		// for LatLong we want to Clamp Y but Wrap X ; that's not supported so just clamp
+		// external caller now prefers to use use ResizeImage for latlongs
+
+		BuildSettings.TextureAddressModeX = TA_Clamp;
+		BuildSettings.TextureAddressModeY = TA_Clamp;
+	}
 
 	FImage Temp;
 	// convert to RGBA32F linear for the compressor
diff --git a/Engine/Source/Runtime/Engine/Private/TextureLODSettings.cpp b/Engine/Source/Runtime/Engine/Private/TextureLODSettings.cpp
--- a/Engine/Source/Runtime/Engine/Private/TextureLODSettings.cpp
+++ b/Engine/Source/Runtime/Engine/Private/TextureLODSettings.cpp
@@ -307,7 +307,9 @@ void UTextureLODSettings::GetMipGenSettings(const UTexture& Texture, TextureMipG
 
 	// ------------
 
-	// default:
+	// default to 2x2 SimpleAverage :
+	// if you generate mips when TMGS was set to NoMipMaps or LeaveExisting, etc.
+	//	it will use these defaults
 	OutSharpen = 0;
 	OutKernelSize = 2;
 
diff --git a/Engine/Source/Runtime/ImageCore/Private/ImageCore.cpp b/Engine/Source/Runtime/ImageCore/Private/ImageCore.cpp
--- a/Engine/Source/Runtime/ImageCore/Private/ImageCore.cpp
+++ b/Engine/Source/Runtime/ImageCore/Private/ImageCore.cpp
@@ -1859,8 +1859,9 @@ static bool GetFormatSTBIR(ERawImageFormat::Type Format,EGammaSpace GammaSpace,
 	}
 }
 
-static bool FilterIsNopWhenSameSize(EResizeImageFilter Filter)
+static bool FilterIsNopWhenSameSize(EResizeImageFilter FilterWithFlags)
 {
+	EResizeImageFilter Filter = FilterWithFlags & EResizeImageFilter::WithoutFlagsMask;
 	switch(Filter)
 	{
 		case EResizeImageFilter::PointSample:
@@ -1883,12 +1884,17 @@ static bool FilterIsNopWhenSameSize(EResizeImageFilter Filter)
 	}
 }
 
-static stbir_filter MapFilterToStb(EResizeImageFilter Filter,int64 SizeFm,int64 SizeTo)
+static stbir_filter MapFilterToStb(EResizeImageFilter FilterWithFlags,int64 SizeFm,int64 SizeTo,uint32 WrapFlag, stbir_edge & OutStbirEdgeMode)
 {
-	if ( SizeFm == SizeTo && FilterIsNopWhenSameSize(Filter) )
+	if ( SizeFm == SizeTo && FilterIsNopWhenSameSize(FilterWithFlags) )
 	{
+		OutStbirEdgeMode = STBIR_EDGE_CLAMP;
 		return STBIR_FILTER_POINT_SAMPLE;
 	}
+	
+	OutStbirEdgeMode = WrapFlag ? STBIR_EDGE_WRAP : STBIR_EDGE_CLAMP;
+	
+	EResizeImageFilter Filter = FilterWithFlags & EResizeImageFilter::WithoutFlagsMask;
 
 	switch(Filter)
 	{
@@ -1964,6 +1970,9 @@ IMAGECORE_API void FImageCore::ResizeImage(const FImageView & SourceImage,const
 		return;
 	}
 	
+	// can't write Pow22 :
+	check( DestImage.GetGammaSpace() != EGammaSpace::Pow22 );
+
 	int SourceNumChannels;
 	stbir_datatype SourceDataType;
 	stbir_pixel_layout SourceLayout;
@@ -2020,12 +2029,14 @@ IMAGECORE_API void FImageCore::ResizeImage(const FImageView & SourceImage,const
 		DestImage.RawData,DestImage.SizeX,DestImage.SizeY,DestImage.GetBytesPerPixel()*DestImage.SizeX,
 		SourceLayout,SourceDataType);
 
-	stbir_filter StbirFilterX = MapFilterToStb(Filter,SourceImage.SizeX,DestImage.SizeX);
-	stbir_filter StbirFilterY = MapFilterToStb(Filter,SourceImage.SizeY,DestImage.SizeY);
+	stbir_edge StbirEdgeX,StbirEdgeY;
+	stbir_filter StbirFilterX = MapFilterToStb(Filter,SourceImage.SizeX,DestImage.SizeX, (uint32)(Filter & EResizeImageFilter::Flag_WrapX), StbirEdgeX);
+	stbir_filter StbirFilterY = MapFilterToStb(Filter,SourceImage.SizeY,DestImage.SizeY, (uint32)(Filter & EResizeImageFilter::Flag_WrapY), StbirEdgeY);
 
 	stbir_set_filters(&resize, StbirFilterX, StbirFilterY);
 	stbir_set_pixel_layouts(&resize,SourceLayout,DestLayout);
 	stbir_set_datatypes(&resize,SourceDataType,DestDataType);
+	stbir_set_edgemodes(&resize, StbirEdgeX, StbirEdgeY);
 
 	const int32 NumWorkers = FTaskGraphInterface::Get().GetNumWorkerThreads();
 
@@ -2070,6 +2081,9 @@ IMAGECORE_API void FImageCore::ResizeImage(const FImageView & SourceImage,const
 IMAGECORE_API void FImageCore::ResizeImageAllocDest(const FImageView & SourceImage,FImage & DestImage,int32 DestSizeX, int32 DestSizeY, ERawImageFormat::Type DestFormat, EGammaSpace DestGammaSpace, EResizeImageFilter Filter)
 {
 	check( DestImage.RawData.Num() == 0 || SourceImage.RawData != DestImage.GetPixelPointer(0,0) ); // must not be resizing onto self
+	
+	// can't write Pow22 :
+	check( DestGammaSpace != EGammaSpace::Pow22 );
 
 	// note that if DestImage was already allocated to the right size, this will not re-allocate it
 	DestImage.Init(DestSizeX,DestSizeY,SourceImage.NumSlices,DestFormat,DestGammaSpace);
@@ -2090,6 +2104,12 @@ IMAGECORE_API void FImageCore::ResizeImageInPlace(FImage & Image,int32 DestSizeX
 		return;
 	}
 
+	// ResizeImageInPlace of Pow22 will change you to sRGB
+	if ( DestGammaSpace == EGammaSpace::Pow22 )
+	{
+		DestGammaSpace = EGammaSpace::sRGB;
+	}
+
 	FImage Source;
 	Image.Swap(Source);
 	ResizeImageAllocDest(Source,Image,DestSizeX,DestSizeY,DestFormat,DestGammaSpace, Filter);
diff --git a/Engine/Source/Runtime/ImageCore/Public/ImageCore.h b/Engine/Source/Runtime/ImageCore/Public/ImageCore.h
--- a/Engine/Source/Runtime/ImageCore/Public/ImageCore.h
+++ b/Engine/Source/Runtime/ImageCore/Public/ImageCore.h
@@ -828,7 +828,7 @@ IMAGECORE_API void TransformToWorkingColorSpace(const FImageView& InLinearImage,
 	/*
 	* filter choice for ResizeImage
 	*/
-	enum class EResizeImageFilter : uint8
+	enum class EResizeImageFilter : uint32
 	{
 		Default = 0, // uses a good default filter; = AdaptiveSharp
 		PointSample,
@@ -839,8 +839,13 @@ IMAGECORE_API void TransformToWorkingColorSpace(const FImageView& InLinearImage,
 		CubicSharp, // sharp interpolating cubic, Catmull-ROM (has negative lobes)
 		CubicMitchell, // compromise between sharp and smooth cubic, Mitchell-Netrevalli filter with B=1/3, C=1/3 (has negative lobes)
 		AdaptiveSharp,  // sharper adaptive filter; uses CubicSharp for upsample and CubicMitchell for downsample, nop for same size
-		AdaptiveSmooth  // smoother adaptive filter; uses CubicMitchell for upsample and CubicGaussian for downsample, nop for same size
+		AdaptiveSmooth,  // smoother adaptive filter; uses CubicMitchell for upsample and CubicGaussian for downsample, nop for same size
+
+		WithoutFlagsMask = 63,
+		Flag_WrapX = 64,  // default edge mode is clamp; set these to wrap instead
+		Flag_WrapY = 128
 	};
+	ENUM_CLASS_FLAGS(EResizeImageFilter);
 
 	/* ResizeImage :
 	*	DestImage should be already allocated; DestImage will be filled in specified format
diff --git a/Engine/Source/Runtime/TextureUtilitiesCommon/Private/TextureSourceDataUtils.cpp b/Engine/Source/Runtime/TextureUtilitiesCommon/Private/TextureSourceDataUtils.cpp
--- a/Engine/Source/Runtime/TextureUtilitiesCommon/Private/TextureSourceDataUtils.cpp
+++ b/Engine/Source/Runtime/TextureUtilitiesCommon/Private/TextureSourceDataUtils.cpp
@@ -17,6 +17,82 @@ namespace UE::TextureUtilitiesCommon::Experimental
 namespace Private
 {
 
+	// adds wrap/clamp flags for X/Y on the filter as appropriate for this texture
+	static FImageCore::EResizeImageFilter SetResizeImageFilterEdgeModes(UTexture * Texture,FImageCore::EResizeImageFilter Filter)
+	{
+		using namespace FImageCore;
+
+		// default Filter with no flags will clamp at edges
+
+		if ( Texture->GetTextureClass() == ETextureClass::Cube ||
+			Texture->GetTextureClass() == ETextureClass::CubeArray )
+		{
+			if ( Texture->Source.IsLongLatCubemap() )
+			{
+				// wrap X, clamp Y
+				return Filter | EResizeImageFilter::Flag_WrapX;
+			}
+			else
+			{
+				// clamp cube faces
+				return Filter;
+			}
+		}
+		else
+		{
+			// see ComputeAddressMode
+
+			if ( Texture->bPreserveBorder )
+			{
+				// clamp
+				return Filter;
+			}
+
+			// @@ nonpow2, UI, NoMipMaps -> clamp ?
+
+			if ( Texture->GetTextureAddressX() == TA_Wrap )
+			{
+				Filter |= EResizeImageFilter::Flag_WrapX;
+			}
+			if ( Texture->GetTextureAddressY() == TA_Wrap )
+			{
+				Filter |= EResizeImageFilter::Flag_WrapY;
+			}
+
+			return Filter;
+		}
+	}
+
+	// use ResizeImage instead of DownsizeImageUsingTextureSettings ?
+	static bool UseResizeImageInsteadOfTextureSettings(UTexture * Texture)
+	{
+		if ( Texture->Source.IsLongLatCubemap() )
+		{
+			// longlat needs to wrap X and clamp Y which the TextureSettings resize can't do
+			return true;
+		}
+
+		if ( Texture->PowerOfTwoMode == ETexturePowerOfTwoSetting::StretchToPowerOfTwo ||
+			Texture->PowerOfTwoMode == ETexturePowerOfTwoSetting::StretchToSquarePowerOfTwo )
+		{
+			if ( ! Texture->Source.AreAllBlocksPowerOfTwo() )
+			{
+				// stretching required, ResizeImage instead
+				return true;
+			}
+		}
+
+		if ( Texture->MipGenSettings == TMGS_NoMipmaps ||
+			Texture->MipGenSettings == TMGS_LeaveExistingMips ||
+			Texture->MipGenSettings == TMGS_Angular )
+		{
+			// TextureSettings mip gen is ill defined in this case, don't use it
+			return true;
+		}
+
+		return false;
+	}
+
 	// resize so that the largest dimension is <= MaxSize
 	static bool ResizeTexture2D(UTexture* Texture, int32 MaxSize, const ITargetPlatform* TargetPlatform)
 	{
@@ -34,18 +110,54 @@ namespace Private
 			return false;
 		}
 
-		bool MadeChanges;
-		if ( ! Texture->DownsizeImageUsingTextureSettings(TargetPlatform, Image, MaxSize, LayerIndex, MadeChanges) )
+		if ( UseResizeImageInsteadOfTextureSettings(Texture) )
 		{
-			UE_LOG(LogTexture,Error,TEXT("ResizeTexture2D: Texture DownsizeImageUsingTextureSettings failed [%s]"),
-				*Texture->GetFullName());
-			return false;
+			int32 TargetSizeX = Image.SizeX;
+			int32 TargetSizeY = Image.SizeY;
+			
+			// stretch to pow2 like TextureCompressorModule
+			if ( Texture->PowerOfTwoMode == ETexturePowerOfTwoSetting::StretchToPowerOfTwo )
+			{
+				TargetSizeX = FMath::RoundUpToPowerOfTwo(Image.SizeX);
+				TargetSizeY = FMath::RoundUpToPowerOfTwo(Image.SizeY);				
+			}
+			else if ( Texture->PowerOfTwoMode == ETexturePowerOfTwoSetting::StretchToSquarePowerOfTwo )
+			{
+				TargetSizeX = TargetSizeY = FMath::RoundUpToPowerOfTwo( FMath::Max(Image.SizeX,Image.SizeY) );
+			}
+
+			// do halving steps to get <= MaxSize
+			while( TargetSizeX > MaxSize || TargetSizeY > MaxSize )
+			{
+				TargetSizeX = FMath::Max(1,TargetSizeX>>1);
+				TargetSizeY = FMath::Max(1,TargetSizeY>>1);
+			}
+
+			if ( TargetSizeX >= Image.SizeX && TargetSizeY >= Image.SizeY )
+			{
+				// nothing to do
+				return false;
+			}
+
+			FImageCore::EResizeImageFilter Filter = SetResizeImageFilterEdgeModes(Texture,FImageCore::EResizeImageFilter::Default);
+
+			FImageCore::ResizeImageInPlace(Image,TargetSizeX,TargetSizeY,Filter);
 		}
-		if ( ! MadeChanges )
+		else
 		{
-			return false;
+			bool MadeChanges = false;
+			if ( ! Texture->DownsizeImageUsingTextureSettings(TargetPlatform, Image, MaxSize, LayerIndex, MadeChanges) )
+			{
+				UE_LOG(LogTexture,Error,TEXT("ResizeTexture2D: Texture DownsizeImageUsingTextureSettings failed [%s]"),
+					*Texture->GetFullName());
+				return false;
+			}
+			if ( ! MadeChanges )
+			{
+				return false;
+			}
 		}
-		
+
 		Texture->PreEditChange(nullptr);
 
 		Texture->Source.Init(Image);
@@ -102,7 +214,9 @@ namespace Private
 			TargetSizeX = RoundToNearestInt32PowerOfTwo( (double) TargetSizeY * Image.SizeX / Image.SizeY );
 		}
 
-		FImageCore::ResizeImageInPlace(Image,TargetSizeX,TargetSizeY);
+		FImageCore::EResizeImageFilter Filter = SetResizeImageFilterEdgeModes(Texture,FImageCore::EResizeImageFilter::Default);
+
+		FImageCore::ResizeImageInPlace(Image,TargetSizeX,TargetSizeY,Filter);
 		
 		Texture->PreEditChange(nullptr);
 
@@ -404,8 +518,6 @@ TEXTUREUTILITIESCOMMON_API bool DownsizeTextureSourceData(UTexture* Texture, int
 		return false;
 	}
 
-	// ?? if Texture has StretchToPow2 set, then do ResizeTextureToNearestPow2 automatically here ?
-
 	if (Texture->Source.GetNumBlocks() == 1)
 	{
 		return Private::ResizeTexture2D(Texture, TargetSourceSize, TargetPlatform);
