Update normal map detection on import to use a view instead of locking the mips directly.

diff --git a/Engine/Plugins/Interchange/Runtime/Source/Pipelines/Private/InterchangeGenericTexturePipeline.cpp b/Engine/Plugins/Interchange/Runtime/Source/Pipelines/Private/InterchangeGenericTexturePipeline.cpp
--- a/Engine/Plugins/Interchange/Runtime/Source/Pipelines/Private/InterchangeGenericTexturePipeline.cpp
+++ b/Engine/Plugins/Interchange/Runtime/Source/Pipelines/Private/InterchangeGenericTexturePipeline.cpp
@@ -74,11 +74,11 @@ namespace UE::Interchange::Private
 	}
 
 #if WITH_EDITOR
-	void AdjustTextureForNormalMap(UTexture* Texture, bool bFlipNormalMapGreenChannel)
+	void AdjustTextureForNormalMap(UTexture* Texture, FImageView MipToAnalyze, bool bFlipNormalMapGreenChannel)
 	{
 		if (Texture)
 		{
-			if (UE::NormalMapIdentification::HandleAssetPostImport(Texture))
+			if (UE::NormalMapIdentification::HandleAssetPostImport(Texture, MipToAnalyze))
 			{
 				UE_LOG(LogInterchangePipeline, Display, TEXT("Auto-detected normal map"));
 
@@ -375,8 +375,10 @@ void UInterchangeGenericTexturePipeline::PostImportTextureAssetImport(UObject* C
 
 	// this is run on main thread
 	check(IsInGameThread());
-	//	after texture may have started compiling
 
+	// (Note - as part of the standard interchange import this is called during the object
+	// import iteration, _before_ the iteration to call to PostEditChange which is what starts the texture build via UpdateResource,
+	// so altering properties here should be safe!)
 	if (!bIsAReimport && bDetectNormalMapTexture)
 	{
 		if (UTexture* Texture = Cast<UTexture>(CreatedAsset))
@@ -384,9 +386,22 @@ void UInterchangeGenericTexturePipeline::PostImportTextureAssetImport(UObject* C
 			// if it's already a normal map, no need to run NormalMapIdentification
 			if (!Texture->IsNormalMap())
 			{
- 				check(!FTextureCompilingManager::Get().IsCompilingTexture(Texture));
-				// AdjustTextureForNormalMap does a PostEditChange which triggers a rebuild
-				UE::Interchange::Private::AdjustTextureForNormalMap(Texture, bFlipNormalMapGreenChannel);
+ 				check(!FTextureCompilingManager::Get().IsCompilingTexture(Texture)); // see comment above.
+
+				FTextureSource::FMipLock LockedMip(FTextureSource::ELockState::ReadOnly, &Texture->Source, 0);
+
+				if (LockedMip.IsValid())
+				{
+					// AdjustTextureForNormalMap technically only adjusts properties and doesn't kick a texture build, however
+					// if it guesses it's a normal map it pops a toast notification that can Revert the change, which does a whole
+					// Modify / PostEditChange which theoretically can kick a build before the PostEditChange in the outer interchange
+					// import chain if somehow the UI click chain routes before the outer loop calls PostEditChange.
+					UE::Interchange::Private::AdjustTextureForNormalMap(Texture, LockedMip.Image, bFlipNormalMapGreenChannel);
+				}
+				else
+				{
+					UE_LOG(LogInterchangePipeline, Display, TEXT("PostImport Texture failed to lock mip data, actions (like normal map detection) not performed!"));
+				}
 			}
 		}
 	}
diff --git a/Engine/Source/Editor/UnrealEd/Private/Factories/EditorFactories.cpp b/Engine/Source/Editor/UnrealEd/Private/Factories/EditorFactories.cpp
--- a/Engine/Source/Editor/UnrealEd/Private/Factories/EditorFactories.cpp
+++ b/Engine/Source/Editor/UnrealEd/Private/Factories/EditorFactories.cpp
@@ -4284,7 +4284,8 @@ UObject* UTextureFactory::FactoryCreateBinary
 
 	if (!bUsingExistingSettings)
 	{
-		if (UE::NormalMapIdentification::HandleAssetPostImport(Texture))
+		FTextureSource::FMipLock LockedMip(FTextureSource::ELockState::ReadOnly, &Texture->Source, 0);
+		if (LockedMip.IsValid() && UE::NormalMapIdentification::HandleAssetPostImport(Texture, LockedMip.Image))
 		{
 			UE_LOG(LogEditorFactories,Display,TEXT("Auto-detected normal map"));
 
diff --git a/Engine/Source/Runtime/Engine/Classes/Engine/Texture.h b/Engine/Source/Runtime/Engine/Classes/Engine/Texture.h
--- a/Engine/Source/Runtime/Engine/Classes/Engine/Texture.h
+++ b/Engine/Source/Runtime/Engine/Classes/Engine/Texture.h
@@ -414,8 +414,11 @@ struct FTextureSource
 		/** Returns a FSharedBuffer that contains the current texture data but cannot be directly modified */
 		const FSharedBuffer& GetDataReadOnly() const { return ReadOnlyReference; }
 
+		UE_DEPRECATED(5.4, "Use GetDataReadWriteView instead")
+		uint8* GetDataReadWrite() { return (uint8*)GetDataReadWriteView().GetData(); }
+
 		/** Returns a pointer that contains the current texture data and can be written to */
-		uint8* GetDataReadWrite();
+		FMutableMemoryView GetDataReadWriteView();
 
 		/** Returns the internal FSharedBuffer and relinquish ownership, used to transfer the data to virtualized bulkdata */
 		FSharedBuffer Release();
@@ -479,12 +482,14 @@ struct FTextureSource
 	};
 	
 	/**
-	* FMipLock to encapsulate a locked mip
-	*  acquires the lock in construct and unlocks in destruct
+	* FMipLock to encapsulate a locked mip - acquires the lock in construct and unlocks in destruct.
+	* 
+	* Be very careful! Locking as ReadOnly will still get you a mutable ImageView! Altering a read only mip
+	* has consequences!
 	*/
 	struct FMipLock
 	{
-		// constructor locks the mip (can fail, pointer will be null)
+		// constructor locks the mip (can fail, check IsValid())
 		ENGINE_API FMipLock(ELockState InLockState,FTextureSource * InTextureSource,int32 InBlockIndex, int32 InLayerIndex, int32 InMipIndex);
 		ENGINE_API FMipLock(ELockState InLockState,FTextureSource * InTextureSource,int32 InMipIndex);
 		
@@ -567,8 +572,8 @@ private:
 	/** Pointer to locked mip data, if any. */
 	FMipAllocation LockedMipData;
 
-	// Internal implementation for locking the mip data, called by LockMipReadOnly or LockMip */
-	uint8* LockMipInternal(int32 BlockIndex, int32 LayerIndex, int32 MipIndex, ELockState RequestedLockState);
+	// Internal implementation for locking the mip data, called by LockMipReadOnly or LockMip.
+	FMutableMemoryView LockMipInternal(int32 BlockIndex, int32 LayerIndex, int32 MipIndex, ELockState RequestedLockState);
 	
 	/** Returns the source data fully decompressed */
 	// ImageWrapperModule is not used
diff --git a/Engine/Source/Runtime/Engine/Private/Texture.cpp b/Engine/Source/Runtime/Engine/Private/Texture.cpp
--- a/Engine/Source/Runtime/Engine/Private/Texture.cpp
+++ b/Engine/Source/Runtime/Engine/Private/Texture.cpp
@@ -2375,14 +2375,14 @@ FTextureSource::FMipLock::FMipLock(ELockState InLockState,FTextureSource * InTex
 	LayerIndex(InLayerIndex),
 	MipIndex(InMipIndex)
 {
-	void * Locked = TextureSource->LockMipInternal(BlockIndex, LayerIndex, MipIndex, LockState);
-	if ( Locked )
+	FMutableMemoryView Locked = TextureSource->LockMipInternal(BlockIndex, LayerIndex, MipIndex, LockState);
+	if ( !Locked.IsEmpty() )
 	{	
 		FTextureSourceBlock Block;
 		TextureSource->GetBlock(BlockIndex, Block);
 		check(MipIndex < Block.NumMips);
 
-		Image.RawData = Locked;
+		Image.RawData = (uint8*)Locked.GetData();
 		Image.SizeX = FMath::Max(Block.SizeX >> MipIndex, 1);
 		Image.SizeY = FMath::Max(Block.SizeY >> MipIndex, 1);
 		Image.NumSlices = TextureSource->GetMippedNumSlices(Block.NumSlices,MipIndex);
@@ -2391,7 +2391,17 @@ FTextureSource::FMipLock::FMipLock(ELockState InLockState,FTextureSource * InTex
 		
 		const int64 MipSizeBytes = TextureSource->CalcMipSize(BlockIndex, LayerIndex, MipIndex);
 
-		check( Image.GetImageSizeBytes() == MipSizeBytes );		
+		if (Image.GetImageSizeBytes() != Locked.GetSize())
+		{
+			// Don't just check on this one since it's actually potential OOB.
+			UE_LOG(LogTexture, Error, TEXT("Locked mip %d / block %d / layer %d has a format expecting %llu bytes but locked data is %llu, failing to lock!"),
+				InMipIndex, InBlockIndex, InLayerIndex, Image.GetImageSizeBytes(), Locked.GetSize());
+			Image = FImage();
+			LockState = ELockState::None;
+			return;
+		}
+
+		check( Image.GetImageSizeBytes() == MipSizeBytes );
 		check( IsValid() );
 	}
 	else
@@ -2432,21 +2442,21 @@ FTextureSource::FMipLock::~FMipLock()
 
 const uint8* FTextureSource::LockMipReadOnly(int32 BlockIndex, int32 LayerIndex, int32 MipIndex)
 {
-	return LockMipInternal(BlockIndex, LayerIndex, MipIndex, ELockState::ReadOnly);
+	return (const uint8*)LockMipInternal(BlockIndex, LayerIndex, MipIndex, ELockState::ReadOnly).GetData();
 }
 
 uint8* FTextureSource::LockMip(int32 BlockIndex, int32 LayerIndex, int32 MipIndex)
 {
-	return LockMipInternal(BlockIndex, LayerIndex, MipIndex, ELockState::ReadWrite);
+	return (uint8*)LockMipInternal(BlockIndex, LayerIndex, MipIndex, ELockState::ReadWrite).GetData();
 }
 
-uint8* FTextureSource::LockMipInternal(int32 BlockIndex, int32 LayerIndex, int32 MipIndex, ELockState RequestedLockState)
+FMutableMemoryView FTextureSource::LockMipInternal(int32 BlockIndex, int32 LayerIndex, int32 MipIndex, ELockState RequestedLockState)
 {
 	TRACE_CPUPROFILER_EVENT_SCOPE(FTextureSource::LockMip);
 
 	checkf(RequestedLockState != ELockState::None, TEXT("Cannot call FTextureSource::LockMipInternal with a RequestedLockState of type ELockState::None"));
 
-	uint8* MipData = nullptr;
+	FMutableMemoryView MipView;
 
 	if (BlockIndex < GetNumBlocks() && LayerIndex < NumLayers && MipIndex < NumMips)
 	{
@@ -2458,30 +2468,32 @@ uint8* FTextureSource::LockMipInternal(int32 BlockIndex, int32 LayerIndex, int32
 
 		if (RequestedLockState == ELockState::ReadOnly)
 		{
-			MipData = const_cast<uint8*>(static_cast<const uint8*>(LockedMipData.GetDataReadOnly().GetData()));
+			// We cast away the const as the ReadOnly wrapper will put it back.
+			FSharedBuffer ReadOnlyMip = LockedMipData.GetDataReadOnly();
+			MipView = FMutableMemoryView((void*)ReadOnlyMip.GetData(), ReadOnlyMip.GetSize());
 		}
 		else
 		{
-			MipData = LockedMipData.GetDataReadWrite();
+			MipView = LockedMipData.GetDataReadWriteView();
 		}
 
-		if ( MipData == nullptr )
+		if ( MipView.IsEmpty() )
 		{
 			// no data, you did not get the lock, do not call Unlock
-			return nullptr;
+			return MipView;
 		}
 		
 		int64 MipOffset = CalcMipOffset(BlockIndex, LayerIndex, MipIndex);
 		int64 MipSize = CalcMipSize(BlockIndex,LayerIndex,MipIndex);
-		if ( MipOffset + MipSize > LockedMipData.GetSize() )
+
+		MipView.MidInline(MipOffset, MipSize);
+		if (MipView.IsEmpty())
 		{
 			UE_LOG(LogTexture,Error,TEXT("Mip Data is too small : %lld < %lld+%lld"), LockedMipData.GetSize(),MipOffset,MipSize); 
 			LockedMipData.Reset();
-			return nullptr;
+			return MipView;
 		}
 
-		MipData += MipOffset;
-
 		if (NumLockedMips == 0)
 		{
 			LockState = RequestedLockState;
@@ -2494,7 +2506,7 @@ uint8* FTextureSource::LockMipInternal(int32 BlockIndex, int32 LayerIndex, int32
 		++NumLockedMips;
 	}
 
-	return MipData;
+	return MipView;
 }
 
 void FTextureSource::UnlockMip(int32 BlockIndex, int32 LayerIndex, int32 MipIndex)
@@ -4093,14 +4105,14 @@ void FTextureSource::FMipAllocation::Reset()
 	ReadWriteBuffer = nullptr;
 }
 
-uint8* FTextureSource::FMipAllocation::GetDataReadWrite()
+FMutableMemoryView FTextureSource::FMipAllocation::GetDataReadWriteView()
 {
 	if (!ReadWriteBuffer.IsValid())
 	{
 		CreateReadWriteBuffer(ReadOnlyReference.GetData(), ReadOnlyReference.GetSize());
 	}
 
-	return ReadWriteBuffer.Get();
+	return FMutableMemoryView(ReadWriteBuffer.Get(), ReadOnlyReference.GetSize());
 }
 
 FSharedBuffer FTextureSource::FMipAllocation::Release()
diff --git a/Engine/Source/Runtime/TextureUtilitiesCommon/Private/NormalMapIdentification.cpp b/Engine/Source/Runtime/TextureUtilitiesCommon/Private/NormalMapIdentification.cpp
--- a/Engine/Source/Runtime/TextureUtilitiesCommon/Private/NormalMapIdentification.cpp
+++ b/Engine/Source/Runtime/TextureUtilitiesCommon/Private/NormalMapIdentification.cpp
@@ -54,151 +54,41 @@ namespace
 	const float RejectedToTakenRatioThreshold = 0.33f;
 }
 
-////////////////////////////////////////////////////////////////////////////////
-// Texture sampler classes
-class NormalMapSamplerBase
-{
-public:
-	NormalMapSamplerBase()
-	: SourceTexture(NULL)
-	, TextureSizeX(0)
-	, TextureSizeY(0)
-	, SourceTextureData(NULL)
-	{
-	}
-
-	~NormalMapSamplerBase()
-	{
-		if ( SourceTextureData != NULL )
-		{
-			SourceTexture->Source.UnlockMip(0);
-		}
-	}
-
-	bool SetSourceTexture( UTexture* Texture )
-	{
-		SourceTexture = Texture;
-		TextureSizeX = Texture->Source.GetSizeX();
-		TextureSizeY = Texture->Source.GetSizeY();
-		SourceTextureData = Texture->Source.LockMipReadOnly(0);
-		return SourceTextureData != nullptr;
-	}
-
-	UTexture* SourceTexture;
-	int64 TextureSizeX;
-	int64 TextureSizeY;
-	const uint8* SourceTextureData;
-};
-
-template<int RIdx, int GIdx, int BIdx, int AIdx> class SampleNormalMapPixel8 : public NormalMapSamplerBase
-{
-public:
-	SampleNormalMapPixel8() {}
-	~SampleNormalMapPixel8() {}
-
-	FLinearColor DoSampleColor( int32 X, int32 Y )
-	{
-		FLinearColor Result;
-		const uint8* PixelToSample = SourceTextureData + ((Y * TextureSizeX + X) * 4);
-
-		const float OneOver255 = 1.0f / 255.0f;
-
-		Result.B = (float)PixelToSample[BIdx] * OneOver255;
-		Result.G = (float)PixelToSample[GIdx] * OneOver255;
-		Result.R = (float)PixelToSample[RIdx] * OneOver255;
-		Result.A = (float)PixelToSample[AIdx] * OneOver255;
-
-		return Result;
-	}
-
-	float ScaleAndBiasComponent( float Value ) const
-	{
-		return Value * 2.0f - 1.0f;
-	}
-};
-typedef SampleNormalMapPixel8<2, 1, 0, 3> SampleNormalMapPixelBGRA8;
-typedef SampleNormalMapPixel8<0, 1, 2, 3> SampleNormalMapPixelRGBA8;
-
-class SampleNormalMapPixelRGBA16 : public NormalMapSamplerBase
-{
-public:
-	SampleNormalMapPixelRGBA16() {}
-	~SampleNormalMapPixelRGBA16() {}
-
-	FLinearColor DoSampleColor( int32 X, int32 Y )
-	{
-		FLinearColor Result;
-		const uint8* PixelToSample = SourceTextureData + ((Y * TextureSizeX + X) * 8);
-
-		const float OneOver65535 = 1.0f / 65535.0f;
-
-		Result.R = (float)((uint16*)PixelToSample)[0] * OneOver65535;
-		Result.G = (float)((uint16*)PixelToSample)[1] * OneOver65535;
-		Result.B = (float)((uint16*)PixelToSample)[2] * OneOver65535;
-		Result.A = (float)((uint16*)PixelToSample)[3] * OneOver65535;
-
-		return Result;
-	}
-
-	float ScaleAndBiasComponent( float Value ) const
-	{
-		return Value * 2.0f - 1.0f;
-	}
-};
-class SampleNormalMapPixelF16 : public NormalMapSamplerBase
-{
-public:
-	SampleNormalMapPixelF16() {}
-	~SampleNormalMapPixelF16() {}
-
-	FLinearColor DoSampleColor( int32 X, int32 Y )
-	{
-		const uint8* PixelToSample = SourceTextureData + ((Y * TextureSizeX + X) * sizeof(FFloat16Color));
-
-		// this assume the normal map to be in linear (not the case if Photoshop converts a 8bit normalmap to float and saves it as 16bit dds)
-		
-		return ((const FFloat16Color*)PixelToSample)->GetFloats();
-	}
-
-	float ScaleAndBiasComponent( float Value ) const
-	{
-		// no need to scale and bias floating point components.
-		return Value;
-	}
-};
-class SampleNormalMapPixelF32 : public NormalMapSamplerBase
-{
-public:
-	SampleNormalMapPixelF32() {}
-	~SampleNormalMapPixelF32() {}
-
-	FLinearColor DoSampleColor( int32 X, int32 Y )
-	{
-		const uint8* PixelToSample = SourceTextureData + ((Y * TextureSizeX + X) * sizeof(FLinearColor));
-		return *((const FLinearColor*)PixelToSample);
-	}
-
-	float ScaleAndBiasComponent( float Value ) const
-	{
-		// no need to scale and bias floating point components.
-		return Value;
-	}
-};
 
 ////////////////////////////////////////////////////////////////////////////////
 ////////////////////////////////////////////////////////////////////////////////
 ////////////////////////////////////////////////////////////////////////////////
-template<class SamplerClass> class TNormalMapAnalyzer
+class FNormalMapAnalyzer
 {
 public:
-	TNormalMapAnalyzer()
+	FNormalMapAnalyzer(const FImageView& InMipToAnalyze)
 	: NumSamplesTaken(0)
 	, NumSamplesRejected(0)
 	, NumSamplesThreshold(0)
 	, AverageColor(0.0f,0.0f,0.0f,0.0f)
+	, SourceImage(InMipToAnalyze)
 	{
 	}
 
+	FVector3f ConvertToVectorForAnalysis(FLinearColor InColor)
+	{
+		// floating point formats don't need to be biased.
+		FVector3f Result;
+		if (SourceImage.Format == ERawImageFormat::BGRA8 ||
+			SourceImage.Format == ERawImageFormat::RGBA16)
+		{
+			Result.X = InColor.R * 2.0f - 1.0f;
+			Result.Y = InColor.G * 2.0f - 1.0f;
+			Result.Z = InColor.B * 2.0f - 1.0f;
+			return Result;
+		}
+		Result.X = InColor.R;
+		Result.Y = InColor.G;
+		Result.Z = InColor.B;
+		return Result;
+	}
+
+
 	/**
 	 * EvaluateSubBlock
 	 * Iterates over all pixels in the specified rectangle, if the resulting pixel
@@ -211,7 +101,7 @@ public:
 		{
 			for ( int32 X=Left; X != (Left+Width); X++ )
 			{
-				FLinearColor ColorSample = Sampler.DoSampleColor( X, Y );
+				FLinearColor ColorSample = SourceImage.GetOnePixelLinear( X, Y );
 
 				// Nearly black or transparent pixels don't contribute to the calculation
 				if (FMath::IsNearlyZero(ColorSample.A, AlphaComponentNearlyZeroThreshold) ||
@@ -221,11 +111,8 @@ public:
 				}
 
 				// Scale and bias, if required, to get a signed vector
-				float Vx = Sampler.ScaleAndBiasComponent( ColorSample.R );
-				float Vy = Sampler.ScaleAndBiasComponent( ColorSample.G );
-				float Vz = Sampler.ScaleAndBiasComponent( ColorSample.B );
-
-				const float Length = FMath::Sqrt(Vx * Vx + Vy * Vy + Vz * Vz);
+				FVector3f Vector = ConvertToVectorForAnalysis(ColorSample);
+				const float Length = Vector.Length();
 				if (Length < ColorComponentNearlyZeroThreshold)
 				{
 					// mid-grey pixels representing (0,0,0) are also not considered as they may be used to denote unused areas
@@ -240,7 +127,7 @@ public:
 				}
 
 				// If the vector is pointing backwards then it is an invalid sample, so consider it invalid
-				if (Vz < 0.0f)
+				if (Vector.Z < 0.0f)
 				{
 					NumSamplesRejected++;
 					continue;
@@ -277,21 +164,16 @@ public:
 	 * value specifications. If the vector satisfies those tolerances then the texture is
 	 * considered to be a normal map.
 	 */
-	bool DoesTextureLookLikelyToBeANormalMap( UTexture* Texture )
+	bool DoesTextureLookLikelyToBeANormalMap()
 	{
-		int32 TextureSizeX = Texture->Source.GetSizeX();
-		int32 TextureSizeY = Texture->Source.GetSizeY();
+		int32 TextureSizeX = SourceImage.GetWidth();
+		int32 TextureSizeY = SourceImage.GetHeight();
 
 		// Calculate the number of tiles in each axis, but limit the number
 		// we interact with to a maximum of 16 tiles (4x4)
 		int32 NumTilesX = FMath::Min( TextureSizeX / SampleTileEdgeLength, MaxTilesPerAxis );
 		int32 NumTilesY = FMath::Min( TextureSizeY / SampleTileEdgeLength, MaxTilesPerAxis );
 
-		if ( ! Sampler.SetSourceTexture( Texture ) )
-		{
-			return false;
-		}
-
 		if (( NumTilesX > 0 ) &&
 			( NumTilesY > 0 ))
 		{
@@ -331,14 +213,12 @@ public:
 			AverageColor /= (float)NumSamplesTaken;
 
 			// See if the resulting vector lies anywhere near the {0,0,1} vector
-			float Vx = Sampler.ScaleAndBiasComponent( AverageColor.R );
-			float Vy = Sampler.ScaleAndBiasComponent( AverageColor.G );
-			float Vz = Sampler.ScaleAndBiasComponent( AverageColor.B );
+			FVector3f Vector = ConvertToVectorForAnalysis(AverageColor);
 
-			float Magnitude = FMath::Sqrt( Vx*Vx + Vy*Vy + Vz*Vz );
+			float Magnitude = Vector.Length();
 
 			// The normalized value of the Z component tells us how close to {0,0,1} the average vector is
-			float NormalizedZ = Vz / Magnitude;
+			float NormalizedZ = Vector.Z / Magnitude;
 
 			// if the average vector is longer than or equal to the min length, shorter than the max length
 			// and the normalized Z value means that the vector is close enough to {0,0,1} then we consider
@@ -357,7 +237,7 @@ public:
 	int32 NumSamplesThreshold;
 	FLinearColor AverageColor;
 
-	SamplerClass Sampler;
+	const FImageView& SourceImage;
 };
 
 /**
@@ -367,9 +247,9 @@ public:
  *
  * @return bool true if the texture is likely a normal map (although it's not necessarily guaranteed)
  */
-static bool IsTextureANormalMap( UTexture* Texture )
+static bool IsImageANormalMap( const FImageView& Image, FStringView TextureDebugName )
 {
-	TRACE_CPUPROFILER_EVENT_SCOPE(IsTextureANormalMap);
+	TRACE_CPUPROFILER_EVENT_SCOPE(IsImageANormalMap);
 
 #if NORMALMAP_IDENTIFICATION_TIMING
 	double StartSeconds = FPlatformTime::Seconds();
@@ -377,46 +257,22 @@ static bool IsTextureANormalMap( UTexture* Texture )
 
 	// Analyze the source texture to try and figure out if it's a normal map.
 	// First check is to make sure it's an appropriate surface format.
-	ETextureSourceFormat SourceFormat = Texture->Source.GetFormat();
-
+	
 	bool bIsNormalMap = false;
-	switch ( SourceFormat )
+	if (Image.Format == ERawImageFormat::BGRA8 ||
+		Image.Format == ERawImageFormat::RGBA16 ||
+		Image.Format == ERawImageFormat::RGBA16F ||
+		Image.Format == ERawImageFormat::RGBA32F)
 	{
 		// The texture could be a normal map if it's one of these formats
-		case TSF_BGRA8:
-			{
-				TNormalMapAnalyzer<SampleNormalMapPixelBGRA8> Analyzer;
-				bIsNormalMap = Analyzer.DoesTextureLookLikelyToBeANormalMap( Texture );
-			}
-			break;
-		case TSF_RGBA16:
-			{
-				TNormalMapAnalyzer<SampleNormalMapPixelRGBA16> Analyzer;
-				bIsNormalMap = Analyzer.DoesTextureLookLikelyToBeANormalMap( Texture );
-			}
-			break;
-		case TSF_RGBA16F:
-			{
-				TNormalMapAnalyzer<SampleNormalMapPixelF16> Analyzer;
-				bIsNormalMap = Analyzer.DoesTextureLookLikelyToBeANormalMap( Texture );
-			}
-			break;
-		case TSF_RGBA32F:
-			{
-				TNormalMapAnalyzer<SampleNormalMapPixelF32> Analyzer;
-				bIsNormalMap = Analyzer.DoesTextureLookLikelyToBeANormalMap( Texture );
-			}
-			break;
-
-		default:
-			// assume the texture is not a normal map
-			break;
+		FNormalMapAnalyzer Analyzer(Image);
+		bIsNormalMap = Analyzer.DoesTextureLookLikelyToBeANormalMap();
 	}
 
 #if NORMALMAP_IDENTIFICATION_TIMING
 	double EndSeconds = FPlatformTime::Seconds();
 
-	FString Msg = FString::Printf( TEXT("NormalMapIdentification took %f seconds to analyze %s"), (EndSeconds-StartSeconds), *Texture->GetFullName() ); 
+	FString Msg = FString::Printf( TEXT("NormalMapIdentification took %f seconds to analyze %.*s"), (EndSeconds-StartSeconds), TextureDebugName.Len(), TextureDebugName.GetData() ); 
 
 	GLog->Log(Msg);
 #endif
@@ -489,13 +345,13 @@ public:
 	TWeakPtr<SNotificationItem> Notification;
 };
 
-bool UE::NormalMapIdentification::HandleAssetPostImport( UTexture* Texture )
+bool UE::NormalMapIdentification::HandleAssetPostImport( UTexture* Texture, const FImageView& InMipToAnalyze )
 {
 	if( Texture != NULL)
 	{
 		// Try to automatically identify a normal map
 		//	this only reads Texture->Source
-		if ( IsTextureANormalMap( Texture ) )
+		if ( IsImageANormalMap( InMipToAnalyze, Texture->GetPathName()) )
 		{
 			// Set the compression settings and no gamma correction for a normal map
 			{
diff --git a/Engine/Source/Runtime/TextureUtilitiesCommon/Public/NormalMapIdentification.h b/Engine/Source/Runtime/TextureUtilitiesCommon/Public/NormalMapIdentification.h
--- a/Engine/Source/Runtime/TextureUtilitiesCommon/Public/NormalMapIdentification.h
+++ b/Engine/Source/Runtime/TextureUtilitiesCommon/Public/NormalMapIdentification.h
@@ -4,7 +4,7 @@
 
 #if WITH_EDITOR
 
-#include "CoreMinimal.h"
+#include "ImageCore.h"
 
 class UTexture;
 class UTextureFactory;
@@ -12,13 +12,17 @@ class UTextureFactory;
 namespace UE::NormalMapIdentification
 {
 	/**
-	 * Handle callback when an asset is imported.
-	 * @param	InFactory	The texture factory being used.
-	 * @param	InTexture	The texture that was imported.
+	 * Checks if the given mip view of the given texture looks like if contains normal map data, and if so
+	 * updates the texture compression settings to be for a normal map. Additionally fires off a UI
+	 * notification with a revert option, as this is a guess and could very well be wrong.
+	 * 
+	 * This does not examine the entire mip but instead checks small chunks across the mip.
+	 * 
+	 * This should be called when texture properties are safe to edit (i.e. between Pre/PostEditChange).
 	 * 
 	 * @return true if the asset was identified as normal
 	 */
-	bool TEXTUREUTILITIESCOMMON_API HandleAssetPostImport( UTexture* InTexture );
+	bool TEXTUREUTILITIESCOMMON_API HandleAssetPostImport( UTexture* InTexture, const FImageView& InMipToAnalyze );
 }
 
 #endif ///WITH_EDITOR
