NormalMapIdentification : fix normal map identification that was broken by previous commit
normal map pixels should be interpreted as linear gamma

#rnx

diff --git a/Engine/Source/Runtime/TextureUtilitiesCommon/Private/NormalMapIdentification.cpp b/Engine/Source/Runtime/TextureUtilitiesCommon/Private/NormalMapIdentification.cpp
--- a/Engine/Source/Runtime/TextureUtilitiesCommon/Private/NormalMapIdentification.cpp
+++ b/Engine/Source/Runtime/TextureUtilitiesCommon/Private/NormalMapIdentification.cpp
@@ -265,7 +265,17 @@ static bool IsImageANormalMap( const FImageView& Image, FStringView TextureDebug
 		Image.Format == ERawImageFormat::RGBA32F)
 	{
 		// The texture could be a normal map if it's one of these formats
-		FNormalMapAnalyzer Analyzer(Image);
+
+		// for BGRA8 sources, interpret them as linear, not SRGB-encoded
+		FImageView LinearImage = Image;
+		if ( LinearImage.GetGammaSpace() != EGammaSpace::Linear )
+		{
+			// note, not converting pixels and copying image
+			//	just reinterpretting
+			LinearImage.GammaSpace = EGammaSpace::Linear;
+		}
+
+		FNormalMapAnalyzer Analyzer(LinearImage);
 		bIsNormalMap = Analyzer.DoesTextureLookLikelyToBeANormalMap();
 	}
 
