Remove some calls to FTextureCompilingManager::Get().FinishCompilation
use BlockOnAnyAsyncBuild instead
wrap Texture changes in PreEdit/PostEdit and don't manually call things like Modify and UpdateResource

#rnx

diff --git a/Engine/Source/Editor/TextureEditor/Private/TextureEditorToolkit.cpp b/Engine/Source/Editor/TextureEditor/Private/TextureEditorToolkit.cpp
--- a/Engine/Source/Editor/TextureEditor/Private/TextureEditorToolkit.cpp
+++ b/Engine/Source/Editor/TextureEditor/Private/TextureEditorToolkit.cpp
@@ -43,7 +43,6 @@
 #include "DeviceProfiles/DeviceProfile.h"
 #include "Curves/CurveLinearColorAtlas.h"
 #include "TextureEditorSettings.h"
-#include "TextureCompiler.h"
 #include "Widgets/Input/SSlider.h"
 #include "Widgets/Input/STextComboBox.h"
 #include "Widgets/Layout/SSpacer.h"
@@ -210,7 +209,7 @@ void FTextureEditorToolkit::InitTextureEditor( const EToolkitMode::Type Mode, co
 
 	// The texture being edited might still be compiling, wait till it finishes then.
 	// FinishCompilation is nice enough to provide a progress for us while we're waiting.
-	FTextureCompilingManager::Get().FinishCompilation({Texture});
+	Texture->BlockOnAnyAsyncBuild();
 
 	// Support undo/redo
 	Texture->SetFlags(RF_Transactional);
diff --git a/Engine/Source/Editor/UnrealEd/Private/ObjectTools.cpp b/Engine/Source/Editor/UnrealEd/Private/ObjectTools.cpp
--- a/Engine/Source/Editor/UnrealEd/Private/ObjectTools.cpp
+++ b/Engine/Source/Editor/UnrealEd/Private/ObjectTools.cpp
@@ -99,7 +99,6 @@
 #include "HAL/PlatformApplicationMisc.h"
 #include "Kismet2/BlueprintEditorUtils.h"
 #include "Subsystems/AssetEditorSubsystem.h"
-#include "TextureCompiler.h"
 #include "UObject/ReferencerFinder.h"
 #include "Containers/Set.h"
 #include "UObject/StrongObjectPtr.h"
@@ -5183,7 +5182,7 @@ namespace ThumbnailTools
 					OutData.AddUninitialized(OutThumbnail->GetImageWidth() * OutThumbnail->GetImageHeight() * sizeof(FColor));
 
 					// Copy the contents of the remote texture to system memory
-					// NOTE: OutRawImageData must be a preallocated buffer!
+					// prefer GetRenderTargetImage()
 					RenderTargetResource->ReadPixelsPtr((FColor*)OutData.GetData(), FReadSurfaceDataFlags(), InSrcRect);
 				}
 			}
@@ -5209,7 +5208,8 @@ namespace ThumbnailTools
 
 			if ( UTexture* Texture = Cast<UTexture>(InObject) )
 			{
-				FTextureCompilingManager::Get().FinishCompilation({Texture});
+				// SetForceMipLevelsToBeResident ?
+				Texture->BlockOnAnyAsyncBuild();
 				Texture->WaitForStreaming();
 			}
 
diff --git a/Engine/Source/Runtime/Engine/Private/TextureRenderTarget.cpp b/Engine/Source/Runtime/Engine/Private/TextureRenderTarget.cpp
--- a/Engine/Source/Runtime/Engine/Private/TextureRenderTarget.cpp
+++ b/Engine/Source/Runtime/Engine/Private/TextureRenderTarget.cpp
@@ -11,7 +11,6 @@
 #include "EngineModule.h"
 #include "EngineLogs.h"
 #include "RenderUtils.h"
-#include "TextureCompiler.h"
 #include "Logging/MessageLog.h"
 #include "Hash/xxhash.h"
 #include "ImageCoreUtils.h"
@@ -407,13 +406,14 @@ bool UTextureRenderTarget::UpdateTexture(UTexture* InTexture, EConstructTextureF
 	{
 		InOnTextureChangingDelegate(InTexture);
 
-		// Ensure the texture is not being compiled
-		FTextureCompilingManager::Get().FinishCompilation({ InTexture });
+		InTexture->PreEditChange(nullptr);
 
 		// init to the same size as the render target
 		InTexture->Source.Init(ReadImage);
 		InTexture->CompressionSettings = CompressionSettingsForTexture;
 		InTexture->SRGB = ReadImage.GammaSpace == EGammaSpace::sRGB;
+
+		InTexture->PostEditChange();
 	}
 
 	return true;
diff --git a/Engine/Source/Runtime/TextureUtilitiesCommon/Private/NormalMapIdentification.cpp b/Engine/Source/Runtime/TextureUtilitiesCommon/Private/NormalMapIdentification.cpp
--- a/Engine/Source/Runtime/TextureUtilitiesCommon/Private/NormalMapIdentification.cpp
+++ b/Engine/Source/Runtime/TextureUtilitiesCommon/Private/NormalMapIdentification.cpp
@@ -6,7 +6,6 @@
 
 #include "Engine/Texture2D.h"
 #include "Framework/Notifications/NotificationManager.h"
-#include "TextureCompiler.h"
 #include "Widgets/Notifications/SNotificationList.h"
 
 #define NORMALMAP_IDENTIFICATION_TIMING	(0)
@@ -318,28 +317,15 @@ public:
 		UTexture2D* Texture2D = Texture.IsValid() ? Cast<UTexture2D>(Texture.Get()) : NULL;
 		if ( Texture2D )
 		{
-			if (FTextureCompilingManager::Get().IsCompilingTexture(Texture2D))
-			{
-				// Block until compile is done
-				TArray<UTexture*> TextureArray;
-				TextureArray.Add(Texture2D);
-				FTextureCompilingManager::Get().FinishCompilation(TextureArray);
-			}
-
 			if ( Texture2D->CompressionSettings == TC_Normalmap )
 			{
-				// Must wait until the texture is done with previous operations before changing settings and getting it to rebuild.
-				Texture2D->WaitForPendingInitOrStreaming();
-
+				Texture2D->PreEditChange(nullptr);
 				Texture2D->SetFlags(RF_Transactional);
-				// Modify calls FinishCachePlatformData to wait on any async build of this texture
-				Texture2D->Modify();
-				Texture2D->PreEditChange(NULL);
-				{
-					Texture2D->CompressionSettings = TC_Default;
-					Texture2D->SRGB = true;
-					Texture2D->LODGroup = TEXTUREGROUP_World;
-				}
+
+				Texture2D->CompressionSettings = TC_Default;
+				Texture2D->LODGroup = TEXTUREGROUP_World;
+				Texture2D->SRGB = true;
+
 				Texture2D->PostEditChange();
 			}
 		}
