Fixing MaterialX imports texture with wrong colors, correctly detect gray scale textures in the Texture Pipeline, not in the Translator wrapper.

diff --git a/Engine/Plugins/Interchange/Runtime/Source/Import/Private/Texture/InterchangeImageWrapperTranslator.cpp b/Engine/Plugins/Interchange/Runtime/Source/Import/Private/Texture/InterchangeImageWrapperTranslator.cpp
--- a/Engine/Plugins/Interchange/Runtime/Source/Import/Private/Texture/InterchangeImageWrapperTranslator.cpp
+++ b/Engine/Plugins/Interchange/Runtime/Source/Import/Private/Texture/InterchangeImageWrapperTranslator.cpp
@@ -164,12 +164,6 @@ TOptional<UE::Interchange::FImportImage> UInterchangeImageWrapperTranslator::Get
 		FImage LoadedImage;
 		if (ImageWrapperModule.DecompressImage(Buffer, Length, LoadedImage))
 		{
-			// Todo interchange: should these payload modification be part of the pipeline, factory or stay there?
-			if (UE::TextureUtilitiesCommon::AutoDetectAndChangeGrayScale(LoadedImage))
-			{
-				UE_LOG(LogInterchangeImport, Display, TEXT("Auto-detected grayscale, image changed to G8"));
-			}
-
 			ETextureSourceFormat TextureFormat = FImageCoreUtils::ConvertToTextureSourceFormat(LoadedImage.Format);
 			bool bSRGB = LoadedImage.GammaSpace != EGammaSpace::Linear;
 
diff --git a/Engine/Plugins/Interchange/Runtime/Source/Pipelines/Private/InterchangeGenericTexturePipeline.cpp b/Engine/Plugins/Interchange/Runtime/Source/Pipelines/Private/InterchangeGenericTexturePipeline.cpp
--- a/Engine/Plugins/Interchange/Runtime/Source/Pipelines/Private/InterchangeGenericTexturePipeline.cpp
+++ b/Engine/Plugins/Interchange/Runtime/Source/Pipelines/Private/InterchangeGenericTexturePipeline.cpp
@@ -28,6 +28,7 @@
 #if WITH_EDITOR
 #include "NormalMapIdentification.h"
 #include "TextureCompiler.h"
+#include "TextureImportUtils.h"
 #include "UDIMUtilities.h"
 #endif //WITH_EDITOR
 
@@ -379,29 +380,41 @@ void UInterchangeGenericTexturePipeline::PostImportTextureAssetImport(UObject* C
 	// (Note - as part of the standard interchange import this is called during the object
 	// import iteration, _before_ the iteration to call to PostEditChange which is what starts the texture build via UpdateResource,
 	// so altering properties here should be safe!)
-	if (!bIsAReimport && bDetectNormalMapTexture)
+	if (!bIsAReimport)
 	{
 		if (UTexture* Texture = Cast<UTexture>(CreatedAsset))
 		{
 			// if it's already a normal map, no need to run NormalMapIdentification
-			if (!Texture->IsNormalMap())
-			{
- 				check(!FTextureCompilingManager::Get().IsCompilingTexture(Texture)); // see comment above.
+			const bool bChangeToNormalMap = bDetectNormalMapTexture && !Texture->IsNormalMap();
 
-				FTextureSource::FMipLock LockedMip(FTextureSource::ELockState::ReadOnly, &Texture->Source, 0);
+			// if it's already a grayscale texture no need to detect it
+			const bool bChangeToGrayScale = bDetectGrayScaleTexture && (Texture->CompressionSettings != TextureCompressionSettings::TC_Grayscale);
 
-				if (LockedMip.IsValid())
-				{
-					// AdjustTextureForNormalMap technically only adjusts properties and doesn't kick a texture build, however
-					// if it guesses it's a normal map it pops a toast notification that can Revert the change, which does a whole
-					// Modify / PostEditChange which theoretically can kick a build before the PostEditChange in the outer interchange
-					// import chain if somehow the UI click chain routes before the outer loop calls PostEditChange.
-					UE::Interchange::Private::AdjustTextureForNormalMap(Texture, LockedMip.Image, bFlipNormalMapGreenChannel);
-				}
-				else
-				{
-					UE_LOG(LogInterchangePipeline, Display, TEXT("PostImport Texture failed to lock mip data, actions (like normal map detection) not performed!"));
-				}
+			if(bChangeToNormalMap || bChangeToGrayScale)
+			{
+					check(!FTextureCompilingManager::Get().IsCompilingTexture(Texture)); // see comment above.
+
+					FTextureSource::FMipLock LockedMip(FTextureSource::ELockState::ReadOnly, &Texture->Source, 0);
+
+					if(LockedMip.IsValid())
+					{
+						// AdjustTextureForNormalMap technically only adjusts properties and doesn't kick a texture build, however
+						// if it guesses it's a normal map it pops a toast notification that can Revert the change, which does a whole
+						// Modify / PostEditChange which theoretically can kick a build before the PostEditChange in the outer interchange
+						// import chain if somehow the UI click chain routes before the outer loop calls PostEditChange.
+						if(bChangeToNormalMap)
+						{
+							UE::Interchange::Private::AdjustTextureForNormalMap(Texture, LockedMip.Image, bFlipNormalMapGreenChannel);
+						}
+						else if(bChangeToGrayScale)
+						{
+							UE::TextureUtilitiesCommon::AutoDetectAndChangeGrayScale(Texture, LockedMip.Image);
+						}
+					}
+					else
+					{
+						UE_LOG(LogInterchangePipeline, Display, TEXT("PostImport Texture failed to lock mip data, actions (like normal map detection) not performed!"));
+					}
 			}
 		}
 	}
diff --git a/Engine/Plugins/Interchange/Runtime/Source/Pipelines/Public/InterchangeGenericTexturePipeline.h b/Engine/Plugins/Interchange/Runtime/Source/Pipelines/Public/InterchangeGenericTexturePipeline.h
--- a/Engine/Plugins/Interchange/Runtime/Source/Pipelines/Public/InterchangeGenericTexturePipeline.h
+++ b/Engine/Plugins/Interchange/Runtime/Source/Pipelines/Public/InterchangeGenericTexturePipeline.h
@@ -35,6 +35,12 @@ public:
 	UPROPERTY(EditAnywhere, BlueprintReadWrite, Category = "Textures", Meta=(EditCondition="bImportTextures"))
 	bool bDetectNormalMapTexture = true;
 
+	/**
+	 * If enable, after a new import a test will be run to see if the texture is a gray scale.
+	 */
+	UPROPERTY(EditAnywhere, BlueprintReadWrite, Category = "Textures", Meta = (EditCondition = "bImportTextures"))
+	bool bDetectGrayScaleTexture = true;
+
 	/** If enabled, the texture's green channel will be inverted for normal maps. */
 	UPROPERTY(EditAnywhere, BlueprintReadWrite, Category = "Textures", Meta=(EditCondition="bImportTextures"))
 	bool bFlipNormalMapGreenChannel = false;
diff --git a/Engine/Source/Runtime/TextureUtilitiesCommon/Private/TextureImportUtils.cpp b/Engine/Source/Runtime/TextureUtilitiesCommon/Private/TextureImportUtils.cpp
--- a/Engine/Source/Runtime/TextureUtilitiesCommon/Private/TextureImportUtils.cpp
+++ b/Engine/Source/Runtime/TextureUtilitiesCommon/Private/TextureImportUtils.cpp
@@ -5,10 +5,44 @@
 #include "ImageCore.h"
 #include "Async/ParallelFor.h"
 
+#if WITH_EDITOR
+#include "Engine/Texture2D.h"
+#include "Framework/Notifications/NotificationManager.h"
+#include "TextureCompiler.h"
+#include "Widgets/Notifications/SNotificationList.h"
+
+#define LOCTEXT_NAMESPACE "GrayScaleIdentification"
+#endif // WITH_EDITOR
+
 namespace UE
 {
 	namespace TextureUtilitiesCommon
 	{
+		template<typename ImageInfo, typename = TEnableIf<std::is_convertible_v<ImageInfo, FImageInfo>>>
+		bool DetectGrayScale(ImageInfo& Image)
+		{
+			if(Image.Format != ERawImageFormat::BGRA8)
+			{
+				return false;
+			}
+
+			// auto-detect gray BGRA8 and change to G8			
+			const TArrayView64<FColor> Colors = Image.AsBGRA8();
+			int64 NumPixels = Image.GetNumPixels();
+
+			for(int64 i = 0; i < NumPixels; i++)
+			{
+				if(Colors[i].A != 255 ||
+				   Colors[i].R != Colors[i].B ||
+				   Colors[i].G != Colors[i].B)
+				{
+					return false;
+				}
+			}
+
+			return true;
+		}
+
 		/**
 		 * Detect the existence of gray scale image in some formats and convert those to a gray scale equivalent image
 		 * 
@@ -16,31 +50,124 @@ namespace UE
 		 */
 		bool AutoDetectAndChangeGrayScale(FImage& Image)
 		{
-			if (Image.Format != ERawImageFormat::BGRA8)
+			const bool bDetectAndChangeGrayScale = DetectGrayScale(Image);
+			if(bDetectAndChangeGrayScale)
 			{
-				return false;
+				Image.ChangeFormat(ERawImageFormat::G8, Image.GammaSpace);
 			}
 
-			// auto-detect gray BGRA8 and change to G8
+			return bDetectAndChangeGrayScale;
+		}
 
-			const FColor* Colors = (const FColor*)Image.RawData.GetData();
-			int64 NumPixels = Image.GetNumPixels();
+#if WITH_EDITOR
+		/** Class to handle callbacks from notifications informing the user a texture was imported as grayscale */
+		class FGrayScaleImportNotificationHandler : public TSharedFromThis<FGrayScaleImportNotificationHandler>
+		{
+		public:
+			FGrayScaleImportNotificationHandler() :
+				Texture(nullptr)
+			{}
+
+			~FGrayScaleImportNotificationHandler() = default;
 
-			for (int64 i = 0; i < NumPixels; i++)
+			/** This method is invoked when the user clicks the "OK" button on the notification */
+			void OKSetting(TSharedPtr<FGrayScaleImportNotificationHandler>)
 			{
-				if (Colors[i].A != 255 ||
-					Colors[i].R != Colors[i].B ||
-					Colors[i].G != Colors[i].B)
+				if(Notification.IsValid())
 				{
-					return false ;
+					Notification.Pin()->SetCompletionState(SNotificationItem::ECompletionState::CS_Success);
+					Notification.Pin()->Fadeout();
 				}
 			}
 
-			// yes, it's gray, do it :
-			Image.ChangeFormat(ERawImageFormat::G8, Image.GammaSpace);
+			/* This method is invoked when the user clicked the "Revert" button on the notification */
+			void RevertSetting(TSharedPtr<FGrayScaleImportNotificationHandler>)
+			{
+				UTexture2D* Texture2D = Texture.IsValid() ? Cast<UTexture2D>(Texture.Get()) : NULL;
+				if(Texture2D)
+				{
+					if(FTextureCompilingManager::Get().IsCompilingTexture(Texture2D))
+					{
+						// Block until compile is done
+						TArray<UTexture*> TextureArray;
+						TextureArray.Add(Texture2D);
+						FTextureCompilingManager::Get().FinishCompilation(TextureArray);
+					}
 
-			return true;
+					if(Texture2D->CompressionSettings == TextureCompressionSettings::TC_Grayscale)
+					{
+						// Must wait until the texture is done with previous operations before changing settings and getting it to rebuild.
+						Texture2D->WaitForPendingInitOrStreaming();
+
+						Texture2D->SetFlags(RF_Transactional);
+						// Modify calls FinishCachePlatformData to wait on any async build of this texture
+						Texture2D->Modify();
+						Texture2D->PreEditChange(NULL);
+						{
+							Texture2D->CompressionSettings = TC_Default;
+							Texture2D->SRGB = true;
+							Texture2D->LODGroup = TEXTUREGROUP_World;
+						}
+						Texture2D->PostEditChange();
+					}
+				}
+
+				if(Notification.IsValid())
+				{
+					Notification.Pin()->SetCompletionState(SNotificationItem::ECompletionState::CS_Success);
+					Notification.Pin()->Fadeout();
+				}
+			}
+
+			TWeakObjectPtr<UTexture> Texture;
+			TWeakPtr<SNotificationItem> Notification;
+		};
+
+		bool AutoDetectAndChangeGrayScale(UTexture* InTexture, FImageView& InMipToAnalyze)
+		{
+			if(DetectGrayScale(InMipToAnalyze))
+			{
+				// Set the compression settings
+				{
+					InTexture->SetFlags(RF_Transactional);
+					InTexture->CompressionSettings = TextureCompressionSettings::TC_Grayscale;
+				}
+
+				// Show the user a notification indicating that this texture will be imported as gray scale.
+				// Offer two options to the user, "OK" dismisses the notification early, "Revert" reverts the settings to that of a diffuse map.
+				// ?? Guess?? this has to be done from main thread only??
+				TSharedPtr<FGrayScaleImportNotificationHandler> GrayScaleNotificationDelegate{ new FGrayScaleImportNotificationHandler };
+				GrayScaleNotificationDelegate->Texture = InTexture;
+
+				// this is a cheat to make sure the notification keeps the callback thing alive while it's active...
+				FText OKText = LOCTEXT("ImportTexture_OKGrayScaleSettings", "OK");
+				FText OKTooltipText = LOCTEXT("ImportTexture_OKTooltip", "Accept gray scale settings");
+				FText RevertText = LOCTEXT("ImportTexture_RevertGrayScaleSettings", "Revert");
+				FText RevertTooltipText = LOCTEXT("ImportTexture_RevertTooltip", "Revert to diffuse map settings");
+
+				FFormatNamedArguments Args;
+				Args.Add(TEXT("TextureName"), FText::FromName(InTexture->GetFName()));
+				FNotificationInfo GrayScaleNotification(FText::Format(LOCTEXT("ImportTexture_IsGrayScale", "Texture {TextureName} was imported as gray scale"), Args));
+				GrayScaleNotification.ButtonDetails.Add(FNotificationButtonInfo(OKText, OKTooltipText, FSimpleDelegate::CreateSP(GrayScaleNotificationDelegate.Get(), &FGrayScaleImportNotificationHandler::OKSetting, GrayScaleNotificationDelegate)));
+				GrayScaleNotification.ButtonDetails.Add(FNotificationButtonInfo(RevertText, RevertTooltipText, FSimpleDelegate::CreateSP(GrayScaleNotificationDelegate.Get(), &FGrayScaleImportNotificationHandler::RevertSetting, GrayScaleNotificationDelegate)));
+				GrayScaleNotification.bFireAndForget = true;
+				GrayScaleNotification.bUseLargeFont = false;
+				GrayScaleNotification.bUseSuccessFailIcons = false;
+				GrayScaleNotification.bUseThrobber = false;
+				GrayScaleNotification.ExpireDuration = 10.0f;
+
+				GrayScaleNotificationDelegate->Notification = FSlateNotificationManager::Get().AddNotification(GrayScaleNotification);
+				if(GrayScaleNotificationDelegate->Notification.IsValid())
+				{
+					GrayScaleNotificationDelegate->Notification.Pin()->SetCompletionState(SNotificationItem::CS_Pending);
+				}
+
+				return true;
+			}
+
+			return false;
 		}
+#endif // WITH_EDITOR
 
 		/**
 		 * This fills any pixels of a texture with have an alpha value of zero and RGB=white,
@@ -571,4 +698,6 @@ namespace UE
 			}
 		}
 	}
-}
\ No newline at end of file
+}
+
+#undef LOCTEXT_NAMESPACE
\ No newline at end of file
diff --git a/Engine/Source/Runtime/TextureUtilitiesCommon/Public/TextureImportUtils.h b/Engine/Source/Runtime/TextureUtilitiesCommon/Public/TextureImportUtils.h
--- a/Engine/Source/Runtime/TextureUtilitiesCommon/Public/TextureImportUtils.h
+++ b/Engine/Source/Runtime/TextureUtilitiesCommon/Public/TextureImportUtils.h
@@ -6,6 +6,8 @@
 #include "Engine/TextureDefines.h"
 
 struct FImage;
+struct FImageView;
+class UTexture;
 
 namespace UE::TextureUtilitiesCommon
 {
@@ -16,6 +18,14 @@ namespace UE::TextureUtilitiesCommon
 	 */
 	TEXTUREUTILITIESCOMMON_API bool AutoDetectAndChangeGrayScale(FImage& Image);
 
+	/**
+	 * Detect the existence of gray scale image in some formats and convert those to a gray scale equivalent image
+	 * Additionally fires off a UI notification with a revert option.
+	 * 
+	 * @return true if the image was converted
+	 */
+	TEXTUREUTILITIESCOMMON_API bool AutoDetectAndChangeGrayScale(UTexture* InTexture, FImageView& InMipToAnalyze);
+
 	/**
 	 * For PNG texture importing, this ensures that any pixels with an alpha value of zero have an RGB
 	 * assigned to them from a neighboring pixel which has non-zero alpha.
