TextureSourceDataUtils nop comments only

diff --git a/Engine/Source/Runtime/TextureUtilitiesCommon/Public/TextureSourceDataUtils.h b/Engine/Source/Runtime/TextureUtilitiesCommon/Public/TextureSourceDataUtils.h
--- a/Engine/Source/Runtime/TextureUtilitiesCommon/Public/TextureSourceDataUtils.h
+++ b/Engine/Source/Runtime/TextureUtilitiesCommon/Public/TextureSourceDataUtils.h
@@ -12,15 +12,18 @@ class ITargetPlatform;
 // Not tested against all the texture types and the api might change or be removed without deprecation
 namespace UE::TextureUtilitiesCommon::Experimental
 {
-	// Try to resize the texture source data so that it is not larger than one build for the specified platform (only Texture2D and TextureCube are supported at the moment)
+	// resize texture source to be <= TargetSourceSize x TargetSourceSize
+	//	Halving (mip) steps are done
+	//	(only Texture2D and TextureCube are supported at the moment)
 	// Note it does not trigger the post edit change after modifying the source texture
-	//	TargetSizeInGame should be a power of 2
 	TEXTUREUTILITIESCOMMON_API bool DownsizeTextureSourceData(UTexture* Texture, int32 TargetSourceSize, const ITargetPlatform* TargetPlatform);
 	
 	// Try to resize the texture source data to a power of two
 	// Note it does not trigger the post edit change after modifying the source texture
 	TEXTUREUTILITIESCOMMON_API bool ResizeTextureSourceDataToNearestPowerOfTwo(UTexture* Texture);
 
+	// Try to resize the texture source data so that it is not larger than one build for the specified platform
+	//	(only Texture2D and TextureCube are supported at the moment)
 	//Note: This function does trigger the post edit change after modifying the source texture
 	TEXTUREUTILITIESCOMMON_API bool DownsizeTextureSourceDataNearRenderingSize(UTexture* Texture, const ITargetPlatform* TargetPlatform);
 	
