TextureSourceDataUtils add AdditionalSourceSizeLimit

#rnx

diff --git a/Engine/Source/Runtime/TextureUtilitiesCommon/Private/TextureSourceDataUtils.cpp b/Engine/Source/Runtime/TextureUtilitiesCommon/Private/TextureSourceDataUtils.cpp
--- a/Engine/Source/Runtime/TextureUtilitiesCommon/Private/TextureSourceDataUtils.cpp
+++ b/Engine/Source/Runtime/TextureUtilitiesCommon/Private/TextureSourceDataUtils.cpp
@@ -529,7 +529,7 @@ TEXTUREUTILITIESCOMMON_API bool DownsizeTextureSourceData(UTexture* Texture, int
 	}
 }
 
-TEXTUREUTILITIESCOMMON_API bool DownsizeTextureSourceDataNearRenderingSize(UTexture* Texture, const ITargetPlatform* TargetPlatform)
+TEXTUREUTILITIESCOMMON_API bool DownsizeTextureSourceDataNearRenderingSize(UTexture* Texture, const ITargetPlatform* TargetPlatform, int32 AdditionalSourceSizeLimit)
 {
 	if ( ! Texture->Source.IsValid() )
 	{
@@ -550,6 +550,8 @@ TEXTUREUTILITIESCOMMON_API bool DownsizeTextureSourceDataNearRenderingSize(UText
 		TargetSourceSize = (1U << FMath::FloorLog2(TargetSizeInGame)) * 2;
 	}
 
+	TargetSourceSize = FMath::Min(TargetSourceSize,AdditionalSourceSizeLimit);
+
 	if (DownsizeTextureSourceData(Texture, TargetSourceSize, TargetPlatform))
 	{
 		Texture->LODBias = 0;
diff --git a/Engine/Source/Runtime/TextureUtilitiesCommon/Public/TextureSourceDataUtils.h b/Engine/Source/Runtime/TextureUtilitiesCommon/Public/TextureSourceDataUtils.h
--- a/Engine/Source/Runtime/TextureUtilitiesCommon/Public/TextureSourceDataUtils.h
+++ b/Engine/Source/Runtime/TextureUtilitiesCommon/Public/TextureSourceDataUtils.h
@@ -22,10 +22,13 @@ namespace UE::TextureUtilitiesCommon::Experimental
 	// Note it does not trigger the post edit change after modifying the source texture
 	TEXTUREUTILITIESCOMMON_API bool ResizeTextureSourceDataToNearestPowerOfTwo(UTexture* Texture);
 
-	// Try to resize the texture source data so that it is not larger than one build for the specified platform
+	// Try to resize the texture source data so that it is not larger than what is built for the specified platform
+	//	eg. if the output build size is 1024 but the source size is 2048, the source will be reduced to 1024
+	//	Texture->LODBias is changed so that the output built size stays the same
 	//	(only Texture2D and TextureCube are supported at the moment)
-	//Note: This function does trigger the post edit change after modifying the source texture
-	TEXTUREUTILITIESCOMMON_API bool DownsizeTextureSourceDataNearRenderingSize(UTexture* Texture, const ITargetPlatform* TargetPlatform);
+	//	AdditionalSourceSizeLimit can be used to apply an extra size limit on the source
+	//Note: This function DOES trigger the post edit change after modifying the source texture
+	TEXTUREUTILITIESCOMMON_API bool DownsizeTextureSourceDataNearRenderingSize(UTexture* Texture, const ITargetPlatform* TargetPlatform, int32 AdditionalSourceSizeLimit = 16384);
 	
 	// ChangeTextureSourceFormat calls Pre/Post edit change
 	//	beware that changing format may change the interpretation of the SRGB bool in Texture
