UDIM index Parser : fix previous restriction to max Y of 9 is too conservative, increase to max Y of 90

diff --git a/Engine/Source/Runtime/TextureUtilitiesCommon/Private/UDIMUtilities.cpp b/Engine/Source/Runtime/TextureUtilitiesCommon/Private/UDIMUtilities.cpp
--- a/Engine/Source/Runtime/TextureUtilitiesCommon/Private/UDIMUtilities.cpp
+++ b/Engine/Source/Runtime/TextureUtilitiesCommon/Private/UDIMUtilities.cpp
@@ -40,12 +40,7 @@ namespace UE
 				}
 			}
 	
-			if ( UdimValue < 1001 )
-			{
-				// UDIM starts with 1001 as the origin
-				return INDEX_NONE;
-			}
-			if ( UdimValue > 1100 )
+			if ( UdimValue < UDIM_Min_Index || UdimValue > UDIM_Max_Index )
 			{
 				return INDEX_NONE;
 			}
@@ -55,12 +50,12 @@ namespace UE
 
 		int32 GetUDIMIndex(int32 BlockX, int32 BlockY)
 		{
-			check( BlockX >= 0 && BlockX <= 9 );
-			check( BlockY >= 0 && BlockY <= 9 );
+			check( BlockX >= 0 && BlockX <= UDIM_Max_X );
+			check( BlockY >= 0 && BlockY <= UDIM_Max_Y );
 
 			int32 Ret = BlockY * 10 + BlockX + 1001;
 
-			check( Ret >= 1001 && Ret <= 1100 );
+			check( Ret >= 1001 && Ret <= UDIM_Max_Index );
 
 			return Ret;
 		}
@@ -103,14 +98,14 @@ namespace UE
 		void ExtractUDIMCoordinates(int32 UDIMIndex, int32& OutBlockX, int32& OutBlockY)
 		{
 			check( UDIMIndex != INDEX_NONE );
-			check( UDIMIndex >= 1001 && UDIMIndex <= 1100 );
+			check( UDIMIndex >= UDIM_Min_Index && UDIMIndex <= UDIM_Max_Index );
 
 			// note 1010 is X=9, Y=0 , that is the tens column is not Y
-			OutBlockX = (UDIMIndex - 1001) % 10;
-			OutBlockY = (UDIMIndex - 1001) / 10;
+			OutBlockX = (UDIMIndex - UDIM_Min_Index) % 10;
+			OutBlockY = (UDIMIndex - UDIM_Min_Index) / 10;
 
-			check( OutBlockX >= 0 && OutBlockX <= 9 );
-			check( OutBlockY >= 0 && OutBlockY <= 9 );
+			check( OutBlockX >= 0 && OutBlockX <= UDIM_Max_X );
+			check( OutBlockY >= 0 && OutBlockY <= UDIM_Max_Y );
 		}
 	}
 }
diff --git a/Engine/Source/Runtime/TextureUtilitiesCommon/Public/UDIMUtilities.h b/Engine/Source/Runtime/TextureUtilitiesCommon/Public/UDIMUtilities.h
--- a/Engine/Source/Runtime/TextureUtilitiesCommon/Public/UDIMUtilities.h
+++ b/Engine/Source/Runtime/TextureUtilitiesCommon/Public/UDIMUtilities.h
@@ -13,10 +13,13 @@ namespace UE
 	namespace TextureUtilitiesCommon
 	{
 		// UDIM = 1001 + X + Y*10
-		// valid UDIM is in "1001" to "1100"
-		//	X and Y are in [0,9]
 		//	note the tens digit is not Y ; eg. 1010 is X=9,Y=0
 
+		constexpr int32 UDIM_Max_X = 9;
+		constexpr int32 UDIM_Max_Y = 90;
+		constexpr int32 UDIM_Min_Index = 1001;
+		constexpr int32 UDIM_Max_Index = UDIM_Min_Index + UDIM_Max_X + UDIM_Max_Y * 10;
+
 		// this Regex is unnecessarily complex and broad
 		//	UDIM must start with "10" then two digits, or "1100"
 		constexpr const TCHAR* DefaultUdimRegexPattern = TEXT(R"((.+?)[._](\d{4})$)");
