Add an EditorPerProjectUserSettings override for Texture import PNG infill setting

This is so UI artists can set that to "always" for UI work without affecting other users.

diff --git a/Engine/Plugins/Interchange/Runtime/Source/Import/Private/Texture/InterchangeImageWrapperTranslator.cpp b/Engine/Plugins/Interchange/Runtime/Source/Import/Private/Texture/InterchangeImageWrapperTranslator.cpp
--- a/Engine/Plugins/Interchange/Runtime/Source/Import/Private/Texture/InterchangeImageWrapperTranslator.cpp
+++ b/Engine/Plugins/Interchange/Runtime/Source/Import/Private/Texture/InterchangeImageWrapperTranslator.cpp
@@ -18,6 +18,7 @@
 #include "Nodes/InterchangeBaseNodeContainer.h"
 #include "Texture/TextureTranslatorUtilities.h"
 #include "TextureImportUtils.h"
+#include "TextureImportUserSettings.h"
 #include "TgaImageSupport.h"
 
 #include UE_INLINE_GENERATED_CPP_BY_NAME(InterchangeImageWrapperTranslator)
@@ -72,8 +73,7 @@ UInterchangeImageWrapperTranslator::UInterchangeImageWrapperTranslator()
 	//  so ensure it is done on game thread before using it from threads
 	// that is guaranteed because the module loader inits all CDOs
 
-	PNGInfill = GetDefault<UTextureImportSettings>()->GetPNGInfillMapDefault();
-	check( PNGInfill != ETextureImportPNGInfill::Default );
+	PNGInfill = UE::TextureUtilitiesCommon::GetPNGInfillSetting();
 }
 
 TArray<FString> UInterchangeImageWrapperTranslator::GetSupportedFormats() const
diff --git a/Engine/Source/Editor/UnrealEd/Private/Factories/EditorFactories.cpp b/Engine/Source/Editor/UnrealEd/Private/Factories/EditorFactories.cpp
--- a/Engine/Source/Editor/UnrealEd/Private/Factories/EditorFactories.cpp
+++ b/Engine/Source/Editor/UnrealEd/Private/Factories/EditorFactories.cpp
@@ -274,6 +274,7 @@
 #include "DesktopPlatformModule.h"
 #include "Interfaces/IMainFrameModule.h"
 #include "TextureImportSettings.h"
+#include "TextureImportUserSettings.h"
 #include "AssetImportTask.h"
 #include "ObjectTools.h"
 
@@ -2973,7 +2974,7 @@ bool UTextureFactory::ImportImage(const uint8* Buffer, int64 Length, FFeedbackCo
 
 			if ( ImageFormat == EImageFormat::PNG )
 			{
-				ETextureImportPNGInfill PNGInfill = GetDefault<UTextureImportSettings>()->GetPNGInfillMapDefault();
+				ETextureImportPNGInfill PNGInfill = UE::TextureUtilitiesCommon::GetPNGInfillSetting();
 
 				if (PNGInfill != ETextureImportPNGInfill::Never)
 				{
diff --git a/Engine/Source/Runtime/TextureUtilitiesCommon/Private/TextureImportUserSettings.cpp b/Engine/Source/Runtime/TextureUtilitiesCommon/Private/TextureImportUserSettings.cpp
new file mode 100644
--- /dev/null
+++ b/Engine/Source/Runtime/TextureUtilitiesCommon/Private/TextureImportUserSettings.cpp
@@ -0,0 +1,28 @@
+// Copyright Epic Games, Inc. All Rights Reserved.
+
+#include "TextureImportUserSettings.h"
+
+#include UE_INLINE_GENERATED_CPP_BY_NAME(TextureImportUserSettings)
+
+UTextureImportUserSettings::UTextureImportUserSettings(const FObjectInitializer& ObjectInitializer)
+	: Super(ObjectInitializer)
+{
+	SectionName = TEXT("Importing");
+}
+
+namespace UE::TextureUtilitiesCommon {
+	TEXTUREUTILITIESCOMMON_API ETextureImportPNGInfill GetPNGInfillSetting()
+	{
+		// Try using per-user setting if present
+		ETextureImportPNGInfill PNGInfill = GetDefault<UTextureImportUserSettings>()->PNGInfill;
+		if (PNGInfill == ETextureImportPNGInfill::Default)
+		{
+			// If no user setting, use project setting if set, legacy settings/defaults if not.
+			PNGInfill = GetDefault<UTextureImportSettings>()->GetPNGInfillMapDefault();
+		}
+
+		// "Default" should've been mapped to a concrete setting now.
+		check(PNGInfill != ETextureImportPNGInfill::Default);
+		return PNGInfill;
+	}
+}
diff --git a/Engine/Source/Runtime/TextureUtilitiesCommon/Public/TextureImportSettings.h b/Engine/Source/Runtime/TextureUtilitiesCommon/Public/TextureImportSettings.h
--- a/Engine/Source/Runtime/TextureUtilitiesCommon/Public/TextureImportSettings.h
+++ b/Engine/Source/Runtime/TextureUtilitiesCommon/Public/TextureImportSettings.h
@@ -32,12 +32,12 @@ enum class ETextureImportFloatingPointFormat : uint8
    By default, this is done OnlyOnBinaryTransparency, not on PNG's with non-binary-transparency alpha channels.
    The PNG format has two different ways of storing alpha, either as 1-bit binary transparency, or as full 8/16 bit alpha channels.
 
-   Used to be set from the TextureImporter/FillPNGZeroAlpha config value.  Setting this option will supercede that.
+   Used to be set from the TextureImporter/FillPNGZeroAlpha config value.  Setting this option will supersede that.
  */
 UENUM()
 enum class ETextureImportPNGInfill : uint8
 {
-	/* Use the legacy default behavior, set from the TextureImporter/FillPNGZeroAlpha config value; default was OnlyOnBinaryTransparency. */
+	/* Use the default behavior. For user settings, this means use project settings. For project settings, it's set from the TextureImporter/FillPNGZeroAlpha config value; default was OnlyOnBinaryTransparency. */
 	Default = 0,
 	/* Never infill RGB, import the PNG exactly as it is stored in the file. */
 	Never,
@@ -81,7 +81,7 @@ public:
 	
 	UPROPERTY(config, EditAnywhere, Category=ImportSettings, meta = (
 		DisplayName = "When to infill RGB in transparent white PNG",
-		ToolTip = "Default behavior is to infill only for binary transparency; this setting may change that to always or never.  Will check TextureImporter/FillPNGZeroAlpha if this is not changed from Default.  This setting is applied to newly imported textures, it does not affect existing textures in the project."))
+		ToolTip = "Default behavior is to infill only for binary transparency; this setting may change that to always or never.  Will check TextureImporter/FillPNGZeroAlpha if this is not changed from Default.  This setting is applied to newly imported textures, it does not affect existing textures in the project. This setting is project-global, prefer the per-user variant in Editor Preferences."))
 	ETextureImportPNGInfill PNGInfill = ETextureImportPNGInfill::Default;
 
 	//~ Begin UObject Interface
diff --git a/Engine/Source/Runtime/TextureUtilitiesCommon/Public/TextureImportUserSettings.h b/Engine/Source/Runtime/TextureUtilitiesCommon/Public/TextureImportUserSettings.h
new file mode 100644
--- /dev/null
+++ b/Engine/Source/Runtime/TextureUtilitiesCommon/Public/TextureImportUserSettings.h
@@ -0,0 +1,34 @@
+// Copyright Epic Games, Inc. All Rights Reserved.
+
+#pragma once
+
+#include "CoreMinimal.h"
+#include "UObject/ObjectMacros.h"
+#include "Engine/EngineTypes.h"
+#include "Engine/DeveloperSettings.h"
+#include "Engine/TextureDefines.h"
+#include "TextureImportSettings.h"
+
+#include "TextureImportUserSettings.generated.h"
+
+struct FPropertyChangedEvent;
+
+UCLASS(config = EditorPerProjectUserSettings, meta=(DisplayName="Texture Import"), MinimalAPI)
+class UTextureImportUserSettings : public UDeveloperSettings
+{
+	GENERATED_UCLASS_BODY()
+
+public:
+	UPROPERTY(config, EditAnywhere, Category=ImportSettings, meta = (
+		DisplayName = "When to infill RGB in transparent white PNG",
+		ToolTip = "Whether to perform infill only for binary transparency, always, or never. If set to 'default', uses global project setting."))
+	ETextureImportPNGInfill PNGInfill = ETextureImportPNGInfill::Default;
+};
+
+namespace UE::TextureUtilitiesCommon
+{
+	/** Resolves PNG infill setting using, in order of preference, per-project user settings,
+	 * project settings, and legacy config settings.
+	 */
+	TEXTUREUTILITIESCOMMON_API ETextureImportPNGInfill GetPNGInfillSetting();
+}
\ No newline at end of file
