Coding Standards Refactor
- Refactored the Image Import Utils to follow better naming conventions.

#rnx

diff --git a/Engine/Source/Runtime/TextureUtilitiesCommon/Private/TextureImportUtils.cpp b/Engine/Source/Runtime/TextureUtilitiesCommon/Private/TextureImportUtils.cpp
--- a/Engine/Source/Runtime/TextureUtilitiesCommon/Private/TextureImportUtils.cpp
+++ b/Engine/Source/Runtime/TextureUtilitiesCommon/Private/TextureImportUtils.cpp
@@ -12,8 +12,8 @@ namespace UE
 {
 	namespace TextureUtilitiesCommon
 	{
-		template<typename TImageClass>
-		bool AutoDetectAndChangeGrayScale(TImageClass& Image)
+		template<typename ImageClassType>
+		bool AutoDetectAndChangeGrayScale(ImageClassType& Image)
 		{
 			if (Image.Format != ERawImageFormat::BGRA8)
 			{
diff --git a/Engine/Source/Runtime/TextureUtilitiesCommon/Public/TextureImportUtils.h b/Engine/Source/Runtime/TextureUtilitiesCommon/Public/TextureImportUtils.h
--- a/Engine/Source/Runtime/TextureUtilitiesCommon/Public/TextureImportUtils.h
+++ b/Engine/Source/Runtime/TextureUtilitiesCommon/Public/TextureImportUtils.h
@@ -12,8 +12,8 @@ namespace UE::TextureUtilitiesCommon
 	 *
 	 * @return true if the image was converted
 	 */
-	template<typename TImageClass>
-	TEXTUREUTILITIESCOMMON_API bool AutoDetectAndChangeGrayScale(TImageClass& Image);
+	template<typename ImageClassType>
+	TEXTUREUTILITIESCOMMON_API bool AutoDetectAndChangeGrayScale(ImageClassType& Image);
 
 	/**
 	 * For PNG texture importing, this ensures that any pixels with an alpha value of zero have an RGB
