Texture Import : make automatic settings for nonpow2 texture imports optional in project TextureImportSettings
Note that nonpow2s will still be set to NeverStream by the existing Texture.cpp logic in this CL

diff --git a/Engine/Plugins/Interchange/Runtime/Source/Import/Private/Texture/InterchangeTextureFactory.cpp b/Engine/Plugins/Interchange/Runtime/Source/Import/Private/Texture/InterchangeTextureFactory.cpp
--- a/Engine/Plugins/Interchange/Runtime/Source/Import/Private/Texture/InterchangeTextureFactory.cpp
+++ b/Engine/Plugins/Interchange/Runtime/Source/Import/Private/Texture/InterchangeTextureFactory.cpp
@@ -1833,38 +1833,6 @@ namespace UE::Interchange::Private::InterchangeTextureFactory
 			ensure(false);
 		}
 
-#if WITH_EDITORONLY_DATA
-		// At some point, these decisions should be left to the pipeline
-		if ( Texture->Source.GetNumBlocks() == 1 && !Texture->Source.IsBlockPowerOfTwo(0))
-		{
-			// try to set some better default options for non-pow2 textures
-
-			// if Texture is not pow2 , change to TMGS_NoMipMaps (if it was default)
-			//   this used to be done by Texture2d.cpp ; it is now optional
-			//	 you can set it back to having mips if you want
-			if (Texture->MipGenSettings == TMGS_FromTextureGroup && Texture->GetTextureClass() == ETextureClass::TwoD )
-			{
-				Texture->MipGenSettings = TMGS_NoMipmaps;
-			}
-
-			if ( ! Texture->Source.IsLongLatCubemap() )
-			{
-				// if Texture is not multiple of 4, change TC to EditorIcon ("UserInterface2D")
-				//	if you do not do this, you might see "Texture forced to uncompressed because size is not a multiple of 4"
-				//  this needs to match the logic in Texture.cpp : GetDefaultTextureFormatName
-				const int32 SizeX = Texture->Source.GetSizeX();
-				const int32 SizeY = Texture->Source.GetSizeY();
-				if ((SizeX & 3) != 0 || (SizeY & 3) != 0)
-				{
-					if (Texture->CompressionSettings == TC_Default) // AutoDXT/BC1
-					{
-						Texture->CompressionSettings = TC_EditorIcon; // "UserInterface2D"
-					}
-				}
-			}
-		}
-#endif // WITH_EDITORONLY_DATA
-
 		// at this point the texture is mostly set up
 		// NormalMapIdentification is not run yet
 	}
diff --git a/Engine/Source/Editor/UnrealEd/Private/Factories/EditorFactories.cpp b/Engine/Source/Editor/UnrealEd/Private/Factories/EditorFactories.cpp
--- a/Engine/Source/Editor/UnrealEd/Private/Factories/EditorFactories.cpp
+++ b/Engine/Source/Editor/UnrealEd/Private/Factories/EditorFactories.cpp
@@ -4019,36 +4019,7 @@ UObject* UTextureFactory::FactoryCreateBinary
 	}
 	
 	Texture->bPreserveBorder		= bPreserveBorder;
-
-	if ( Texture->Source.GetNumBlocks() == 1 && ! Texture->Source.IsBlockPowerOfTwo(0) )
-	{
-		// try to set some better default options for non-pow2 textures
-
-		// if Texture is not pow2 , change to TMGS_NoMipMaps (if it was default)
-		//   this used to be done by Texture2d.cpp ; it is now optional
-		//	 you can set it back to having mips if you want
-		if ( Texture->MipGenSettings == TMGS_FromTextureGroup && Texture->GetTextureClass() == ETextureClass::TwoD )
-		{
-			Texture->MipGenSettings = TMGS_NoMipmaps;		
-		}
-		
-		if ( ! Texture->Source.IsLongLatCubemap() )
-		{
-			// if Texture is not multiple of 4, change TC to EditorIcon ("UserInterface2D")
-			//	if you do not do this, you might see "Texture forced to uncompressed because size is not a multiple of 4"
-			//  this needs to match the logic in Texture.cpp : GetDefaultTextureFormatName
-			int32 SizeX = Texture->Source.GetSizeX();
-			int32 SizeY = Texture->Source.GetSizeY();
-			if ( (SizeX&3) != 0 || (SizeY&3) != 0 )
-			{
-				if ( Texture->CompressionSettings == TC_Default ) // AutoDXT/BC1
-				{
-					Texture->CompressionSettings = TC_EditorIcon; // "UserInterface2D"
-				}
-			}
-		}
-	}
-				
+					
 	if ( UE::TextureUtilitiesCommon::ShouldTextureBeVirtualByAutoImportSize(Texture) )
 	{
 		// only UTexture2D can be VT
diff --git a/Engine/Source/Runtime/TextureUtilitiesCommon/Private/TextureImportSettings.cpp b/Engine/Source/Runtime/TextureUtilitiesCommon/Private/TextureImportSettings.cpp
--- a/Engine/Source/Runtime/TextureUtilitiesCommon/Private/TextureImportSettings.cpp
+++ b/Engine/Source/Runtime/TextureUtilitiesCommon/Private/TextureImportSettings.cpp
@@ -126,6 +126,12 @@ namespace UE::TextureUtilitiesCommon
 		{
 			Texture->UpdateOodleTextureSdkVersionToLatest();
 
+			// @todo Oodle : optional : reimport ; if texture was previously nonpow2 but is now pow2,
+			//		change some flags (UI, NeverStream, NoMipMaps, etc.) that may have been set by bDoAutomaticTextureSettingsForNonPow2Textures
+			//		?this may stomp intentional user settings, but maybe that's just okay?
+			//	ideally we'd like to just pop up a query box to say "yes/no" to that
+			//	but we aren't allowed to do that here because we may be running on task threads in an Interchange mass-reimport
+
 			return;
 		}
 
@@ -191,6 +197,38 @@ namespace UE::TextureUtilitiesCommon
 				checkNoEntry();
 			}
 		}
+
+		if ( Settings->bDoAutomaticTextureSettingsForNonPow2Textures )
+		{
+			if ( Texture->Source.GetNumBlocks() == 1 && ! Texture->Source.IsBlockPowerOfTwo(0) )
+			{
+				// try to set some better default options for non-pow2 textures
+
+				// if Texture is not pow2 , change to TMGS_NoMipMaps (if it was default)
+				//   this used to be done by Texture2d.cpp ; it is now optional
+				//	 you can set it back to having mips if you want
+				if ( Texture->MipGenSettings == TMGS_FromTextureGroup && Texture->GetTextureClass() == ETextureClass::TwoD )
+				{
+					Texture->MipGenSettings = TMGS_NoMipmaps;		
+				}
+		
+				if ( ! Texture->Source.IsLongLatCubemap() )
+				{
+					// if Texture is not multiple of 4, change TC to EditorIcon ("UserInterface2D")
+					//	if you do not do this, you might see "Texture forced to uncompressed because size is not a multiple of 4"
+					//  this needs to match the logic in Texture.cpp : GetDefaultTextureFormatName
+					int32 SizeX = Texture->Source.GetSizeX();
+					int32 SizeY = Texture->Source.GetSizeY();
+					if ( (SizeX&3) != 0 || (SizeY&3) != 0 )
+					{
+						if ( Texture->CompressionSettings == TC_Default ) // AutoDXT/BC1
+						{
+							Texture->CompressionSettings = TC_EditorIcon; // "UserInterface2D"
+						}
+					}
+				}
+			}
+		}
 	}
 	
 	bool ShouldTextureBeVirtualByAutoImportSize(const UTexture * Texture)
diff --git a/Engine/Source/Runtime/TextureUtilitiesCommon/Public/TextureImportSettings.h b/Engine/Source/Runtime/TextureUtilitiesCommon/Public/TextureImportSettings.h
--- a/Engine/Source/Runtime/TextureUtilitiesCommon/Public/TextureImportSettings.h
+++ b/Engine/Source/Runtime/TextureUtilitiesCommon/Public/TextureImportSettings.h
@@ -83,6 +83,11 @@ public:
 		DisplayName = "When to infill RGB in transparent white PNG",
 		ToolTip = "Default behavior is to infill only for binary transparency; this setting may change that to always or never.  Will check TextureImporter/FillPNGZeroAlpha if this is not changed from Default.  This setting is applied to newly imported textures, it does not affect existing textures in the project. This setting is project-global, prefer the per-user variant in Editor Preferences."))
 	ETextureImportPNGInfill PNGInfill = ETextureImportPNGInfill::Default;
+	
+	UPROPERTY(config, EditAnywhere, Category=ImportSettings, meta = (
+		DisplayName = "Detect nonpow2 textures and set defaults",
+		ToolTip = "Sets them to NoMipMaps and uncompressed UI compression. If this is off, artists must imported UI textures up correctly by hand."))
+	bool bDoAutomaticTextureSettingsForNonPow2Textures = true;
 
 	//~ Begin UObject Interface
 
