TextureImportSettings : at reimport time, if texture was previously set to UI/Nonstreaming because it was nonpow2, but now the reimport is pow2, automatically change defaults to pow2 standard

diff --git a/Engine/Source/Runtime/TextureUtilitiesCommon/Private/TextureImportSettings.cpp b/Engine/Source/Runtime/TextureUtilitiesCommon/Private/TextureImportSettings.cpp
--- a/Engine/Source/Runtime/TextureUtilitiesCommon/Private/TextureImportSettings.cpp
+++ b/Engine/Source/Runtime/TextureUtilitiesCommon/Private/TextureImportSettings.cpp
@@ -126,11 +126,32 @@ namespace UE::TextureUtilitiesCommon
 		{
 			Texture->UpdateOodleTextureSdkVersionToLatest();
 
-			// @todo Oodle : optional : reimport ; if texture was previously nonpow2 but is now pow2,
+			// optional : at reimport ; if texture was previously nonpow2 but is now pow2,
 			//		change some flags (UI, NeverStream, NoMipMaps, etc.) that may have been set by bDoAutomaticTextureSettingsForNonPow2Textures
 			//		?this may stomp intentional user settings, but maybe that's just okay?
 			//	ideally we'd like to just pop up a query box to say "yes/no" to that
 			//	but we aren't allowed to do that here because we may be running on task threads in an Interchange mass-reimport
+			
+			if ( Texture->GetTextureClass() == ETextureClass::TwoD &&
+				Texture->MipGenSettings == TMGS_NoMipmaps &&
+				( Texture->CompressionSettings == TC_EditorIcon || Texture->CompressionSettings == TC_Default ) &&
+				Texture->Source.IsValid() &&
+				Texture->LODGroup != TEXTUREGROUP_UI )
+			{
+				// AreAllBlocksPowerOfTwo only checks XY but we are on TwoD class here only
+				bool bIsPowerOfTwo = Texture->Source.AreAllBlocksPowerOfTwo() || ( Texture->PowerOfTwoMode != ETexturePowerOfTwoSetting::None );
+
+				if ( bIsPowerOfTwo )
+				{
+					UE_LOG(LogCore, Display, TEXT("Texture [%s] settings changed by reimport; was NoMips UI, is now pow2 Default"), 
+						*Texture->GetName() );	
+
+					// change to standard defaults for pow2 textures
+					Texture->MipGenSettings = TMGS_FromTextureGroup;
+					Texture->CompressionSettings = TC_Default;
+					Texture->NeverStream = false;
+				}
+			}
 
 			return;
 		}
