[Texture Import] UDIM File Matching Filter Fix
- Without this fix, udim prefixes that are substrings of other files in the same directory override the UDIMs selected.
- Now the filter will strictly match the udim prefix from the filename that user tried importing.

#rnx

#Synced-CL 50908799

diff --git a/Engine/Source/Runtime/TextureUtilitiesCommon/Private/UDIMUtilities.cpp b/Engine/Source/Runtime/TextureUtilitiesCommon/Private/UDIMUtilities.cpp
--- a/Engine/Source/Runtime/TextureUtilitiesCommon/Private/UDIMUtilities.cpp
+++ b/Engine/Source/Runtime/TextureUtilitiesCommon/Private/UDIMUtilities.cpp
@@ -62,33 +62,36 @@ namespace UE
 
 		TMap<int32, FString> GetUDIMBlocksFromSourceFile(const FString& File, const FString& UdimRegexPattern, FString* OutFilenameWithoutUdimPatternAndExtension)
 		{
-			FString PreUDIMName;
+			FString PreUDIMNameToMatch;
 			FString PostUDIMName;
 			TMap<int32, FString> UDIMsAndSourcesFile;
 
-			if (ParseUDIMName(FPaths::GetBaseFilename(File), UdimRegexPattern, PreUDIMName, PostUDIMName) != INDEX_NONE)
+			if (ParseUDIMName(FPaths::GetBaseFilename(File), UdimRegexPattern, PreUDIMNameToMatch, PostUDIMName) != INDEX_NONE)
 			{ 
 				const FString Path = FPaths::GetPath(File);
-				const FString UDIMFilter = (Path / PreUDIMName) + TEXT("*") + PostUDIMName + FPaths::GetExtension(File, true);
+				const FString UDIMFilter = (Path / PreUDIMNameToMatch) + TEXT("*") + PostUDIMName + FPaths::GetExtension(File, true);
 
 				TArray<FString> UDIMFiles;
 				IFileManager::Get().FindFiles(UDIMFiles, *UDIMFilter, true, false);
-				UDIMFiles.Reserve(UDIMFiles.Num());
-
+				
+				FString PreUDIMNameFound;
 				for (FString& UDIMFile : UDIMFiles)
 				{
-					const int32 UDIMIndex = ParseUDIMName(FPaths::GetBaseFilename(UDIMFile), UdimRegexPattern, PreUDIMName, PostUDIMName);
-					if (UDIMIndex != INDEX_NONE)
+					const int32 UDIMIndex = ParseUDIMName(FPaths::GetBaseFilename(UDIMFile), UdimRegexPattern, PreUDIMNameFound, PostUDIMName);
+					
+					if (UDIMIndex == INDEX_NONE || !PreUDIMNameFound.Equals(PreUDIMNameToMatch))
 					{
-						const FString UDIMName = PreUDIMName + PostUDIMName;
-						FString UDIMPath = FPaths::Combine(Path, UDIMFile);
-						UDIMsAndSourcesFile.Add(UDIMIndex, MoveTemp(UDIMPath));
+						continue;
 					}
+
+					const FString UDIMName = PreUDIMNameFound + PostUDIMName;
+					FString UDIMPath = FPaths::Combine(Path, UDIMFile);
+					UDIMsAndSourcesFile.Add(UDIMIndex, MoveTemp(UDIMPath));
 				}
 
 				if (!UDIMsAndSourcesFile.IsEmpty() && OutFilenameWithoutUdimPatternAndExtension)
 				{
-					*OutFilenameWithoutUdimPatternAndExtension = PreUDIMName + PostUDIMName;
+					*OutFilenameWithoutUdimPatternAndExtension = PreUDIMNameToMatch + PostUDIMName;
 				}
 			}
 
